* Updating the SClassViewer to no longer use the EditorClassHierarchy, ran into issues with it looking at old stale data.  The new approach builds the nodes using the class iterator.
* Removing the EditorClassHierarchy, as it is now no longer used in the codebase.

diff --git a/Engine/Source/Editor/ClassViewer/Private/ClassViewerNode.h b/Engine/Source/Editor/ClassViewer/Private/ClassViewerNode.h
--- a/Engine/Source/Editor/ClassViewer/Private/ClassViewerNode.h
+++ b/Engine/Source/Editor/ClassViewer/Private/ClassViewerNode.h
@@ -48,6 +48,11 @@ public:
 		return bIsClassPlaceable;
 	}
 
+	void IsClassPlaceable(bool bInIsClassPlaceable)
+	{
+		bIsClassPlaceable = bInIsClassPlaceable;
+	}
+
 	bool IsRestricted() const;
 
 private:
diff --git a/Engine/Source/Editor/ClassViewer/Private/SClassViewer.cpp b/Engine/Source/Editor/ClassViewer/Private/SClassViewer.cpp
--- a/Engine/Source/Editor/ClassViewer/Private/SClassViewer.cpp
+++ b/Engine/Source/Editor/ClassViewer/Private/SClassViewer.cpp
@@ -374,10 +374,9 @@ public:
 private:
 	/** Recursive function to build a tree, will not filter.
 	 *	@param InOutRootNode						The node that this function will add the children of to the tree.
-	 *	@param InRootClassIndex						The index of the root node.
 	 *  @param PackageNameToAssetDataMap			The asset registry map of blueprint package names to blueprint data
 	 */	
-	void AddChildren_NoFilter( TSharedPtr< FClassViewerNode >& InOutRootNode, int32 InRootClassIndex, const TMultiMap<FName,FAssetData>& BlueprintPackageToAssetDataMap );
+	void AddChildren_NoFilter( TSharedPtr< FClassViewerNode >& InOutRootNode, const TMultiMap<FName,FAssetData>& BlueprintPackageToAssetDataMap );
 
 	/** Finds the node, recursively going deeper into the hierarchy. Does so by comparing generated class package names.
 	 *	@param InGeneratedClassPackageName			The name of the generated class package to find the node for.
@@ -785,9 +784,6 @@ namespace ClassViewer
 		 */
 		static void GetClassInfo( TWeakObjectPtr<UClass> InClass, bool& bInOutIsBlueprintBase, bool& bInOutHasBlueprint )
 		{
-			FEditorClassHierarchy* EditorClassHierarchy = GEditor->EditorClassHierarchy;
-			check(EditorClassHierarchy);
-			 
 			if (UClass* Class = InClass.Get())
 			{
 				bInOutIsBlueprintBase = FKismetEditorUtilities::CanCreateBlueprintOfClass( Class );
@@ -1515,11 +1511,16 @@ private:
 	FOnClassItemDoubleClickDelegate OnDoubleClicked;
 };
 
+static void OnModulesChanged(FName ModuleThatChanged, EModuleChangeReason::Type ReasonForChange)
+{
+	ClassViewer::Helpers::RequestPopulateClassHierarchy();
+}
+
 FClassHierarchy::FClassHierarchy()
 {
 	// Register with the Asset Registry to be informed when it is done loading up files.
 	FAssetRegistryModule& AssetRegistryModule = FModuleManager::GetModuleChecked<FAssetRegistryModule>(TEXT("AssetRegistry"));
-	AssetRegistryModule.Get().OnFilesLoaded().AddStatic(ClassViewer::Helpers::RequestPopulateClassHierarchy);
+	AssetRegistryModule.Get().OnFilesLoaded().AddStatic( ClassViewer::Helpers::RequestPopulateClassHierarchy );
 	AssetRegistryModule.Get().OnAssetAdded().AddRaw( this, &FClassHierarchy::AddAsset);
 	AssetRegistryModule.Get().OnAssetRemoved().AddRaw( this, &FClassHierarchy::RemoveAsset );
 
@@ -1527,6 +1528,8 @@ FClassHierarchy::FClassHierarchy()
 	GEditor->OnHotReload().AddStatic(ClassViewer::Helpers::RequestPopulateClassHierarchy);
 	GEditor->OnBlueprintCompiled().AddStatic(ClassViewer::Helpers::RequestPopulateClassHierarchy);
 	GEditor->OnClassPackageLoadedOrUnloaded().AddStatic(ClassViewer::Helpers::RequestPopulateClassHierarchy);
+
+	FModuleManager::Get().OnModulesChanged().AddStatic(&OnModulesChanged);
 }
 
 FClassHierarchy::~FClassHierarchy()
@@ -1544,65 +1547,96 @@ FClassHierarchy::~FClassHierarchy()
 		GEditor->OnBlueprintCompiled().RemoveStatic(ClassViewer::Helpers::RequestPopulateClassHierarchy);
 		GEditor->OnClassPackageLoadedOrUnloaded().RemoveStatic(ClassViewer::Helpers::RequestPopulateClassHierarchy);
 	}
+
+	FModuleManager::Get().OnModulesChanged().RemoveAll(this);
 }
 
-void FClassHierarchy::AddChildren_NoFilter( TSharedPtr< FClassViewerNode >& InOutRootNode, int32 InRootClassIndex, const TMultiMap<FName,FAssetData>& BlueprintPackageToAssetDataMap )
+static TSharedPtr< FClassViewerNode > CreateNodeForClass(UClass* Class, const TMultiMap<FName, FAssetData>& BlueprintPackageToAssetDataMap)
 {
-	FEditorClassHierarchy* EditorClassHierarchy = GEditor->EditorClassHierarchy;
-	check(EditorClassHierarchy);
+	FString ClassName = Class->GetName();
 
-	// Get a list of all direct descendants of this class.
-	TArray<int32> ChildClassIndexArray;
-	EditorClassHierarchy->GetChildrenOfClass(InRootClassIndex, ChildClassIndexArray);
+	const bool bIsPlaceable = ClassViewer::Helpers::IsPlaceable(Class);
+	const bool bIsAbstract = ClassViewer::Helpers::IsAbstract(Class);
+	const bool bIsBrush = ClassViewer::Helpers::IsBrush(Class);
 
-	// Go through all of the classes children and see if they should be added to the list.
-	for(int32 ChildClassIndex = 0; ChildClassIndex < ChildClassIndexArray.Num(); ChildClassIndex++)
-	{
-		int32 CurrentChildIndex = ChildClassIndexArray[ ChildClassIndex ];
+	// Create the new node so it can be passed to AddChildren, fill it in with if it is placeable, abstract, and/or a brush.
+	TSharedPtr< FClassViewerNode > NewNode = MakeShareable(new FClassViewerNode(ClassName, bIsPlaceable && !bIsAbstract && !bIsBrush));
+	NewNode->Blueprint = ClassViewer::Helpers::GetBlueprint(Class);
+	NewNode->Class = Class;
 
-		// Pull out the class for this object, check to make sure it is not NULL.
-		UClass* Class = EditorClassHierarchy->GetClass(CurrentChildIndex);
+	// Retrieve all blueprint classes
+	TArray<FAssetData> BlueprintList;
+	BlueprintPackageToAssetDataMap.MultiFind(NewNode->Class->GetOuterUPackage()->GetFName(), BlueprintList);
+
+	// Check if the Asset Registry found anything, it should, but check.
+	if ( BlueprintList.Num() )
+	{
+		// Grab the generated class name and check it before assigning. Objects that haven't been saved since this has started to be exported do not have the information.
+		FString* GeneratedClassnamePtr = BlueprintList[0].TagsAndValues.Find(FName("GeneratedClass"));
 
-		if(Class)
+		if ( GeneratedClassnamePtr )
 		{
-			FString ClassName = EditorClassHierarchy->GetClassName( CurrentChildIndex );
+			NewNode->GeneratedClassname = *GeneratedClassnamePtr;
+		}
+	}
 
-			const bool bIsPlaceable = ClassViewer::Helpers::IsPlaceable(Class);
-			const bool bIsAbstract = ClassViewer::Helpers::IsAbstract(Class);
-			const bool bIsBrush = ClassViewer::Helpers::IsBrush(Class);
+	return NewNode;
+}
 
-			FString ChildClassName = EditorClassHierarchy->GetClassName(CurrentChildIndex);
+void FClassHierarchy::AddChildren_NoFilter( TSharedPtr< FClassViewerNode >& InOutRootNode, const TMultiMap<FName, FAssetData>& BlueprintPackageToAssetDataMap )
+{
+	UClass* RootClass = UObject::StaticClass();
 
-			// Create the new node so it can be passed to AddChildren, fill it in with if it is placeable, abstract, and/or a brush.
-			TSharedPtr< FClassViewerNode > NewNode = MakeShareable(new FClassViewerNode(ClassName, bIsPlaceable && !bIsAbstract && !bIsBrush));
-			NewNode->Blueprint = ClassViewer::Helpers::GetBlueprint(Class);
-			NewNode->Class = Class;
+	ObjectClassRoot = MakeShareable(new FClassViewerNode(RootClass->GetName(), false));
+	ObjectClassRoot->Class = RootClass;
 
-			// Recurse into the child and check if it should be displayed.
-			AddChildren_NoFilter(NewNode, CurrentChildIndex, BlueprintPackageToAssetDataMap);
+	TMap< UClass*, TSharedPtr< FClassViewerNode > > Nodes;
 
-			// Retrieve all blueprint classes
-			TArray<FAssetData> BlueprintList;
-			BlueprintPackageToAssetDataMap.MultiFind(NewNode->Class->GetOuterUPackage()->GetFName(), BlueprintList);
+	Nodes.Add(RootClass, ObjectClassRoot);
 
-			// Check if the Asset Registry found anything, it should, but check.
-			if (BlueprintList.Num() )
-			{
-				// Grab the generated class name and check it before assigning. Objects that haven't been saved since this has started to be exported do not have the information.
-				FString* GeneratedClassnamePtr = BlueprintList[0].TagsAndValues.Find(FName("GeneratedClass"));
+	TSet<UClass*> Visited;
+	Visited.Add(RootClass);
 
-				if(GeneratedClassnamePtr)
-				{
-					NewNode->GeneratedClassname = *GeneratedClassnamePtr;
-				}
-			}
+	// Go through all of the classes children and see if they should be added to the list.
+	for ( TObjectIterator<UClass> ClassIt; ClassIt; ++ClassIt )
+	{
+		UClass* CurrentClass = *ClassIt;
 
-			InOutRootNode->AddChild(NewNode);
+		// Ignore deprecated and temporary trash classes.
+		if ( CurrentClass->HasAnyClassFlags(CLASS_Deprecated | CLASS_NewerVersionExists) )
+		{
+			continue;
+		}
+		
+		TSharedPtr<FClassViewerNode>& Entry = Nodes.FindOrAdd(CurrentClass);
+		if ( Visited.Contains(CurrentClass) )
+		{
+			continue;
 		}
 		else
 		{
-			// Log that the Class was NULL, ideally no class in the hierarchy should be NULL.
-			UE_LOG(LogEditorClassViewer, Warning, TEXT("Class is NULL"));
+			while ( CurrentClass->GetSuperClass() != NULL )
+			{
+				TSharedPtr<FClassViewerNode>& ParentEntry = Nodes.FindOrAdd(CurrentClass->GetSuperClass());
+				if ( !ParentEntry.IsValid() )
+				{
+					ParentEntry = CreateNodeForClass(CurrentClass->GetSuperClass(), BlueprintPackageToAssetDataMap);
+				}
+
+				TSharedPtr<FClassViewerNode>& MyEntry = Nodes.FindOrAdd(CurrentClass);
+				if ( !MyEntry.IsValid() )
+				{
+					MyEntry = CreateNodeForClass(CurrentClass, BlueprintPackageToAssetDataMap);
+				}
+
+				if ( !Visited.Contains(CurrentClass) )
+				{
+					ParentEntry->AddChild(MyEntry);
+					Visited.Add(CurrentClass);
+				}
+
+				CurrentClass = CurrentClass->GetSuperClass();
+			}
 		}
 	}
 }
@@ -1880,7 +1914,7 @@ void FClassHierarchy::PopulateClassHierarchy()
 	AssetRegistryModule.Get().GetAssets(Filter, BlueprintList);
 
 	TMultiMap<FName,FAssetData> BlueprintPackageToAssetDataMap;
-	for (int32 AssetIdx = 0; AssetIdx < BlueprintList.Num(); ++AssetIdx)
+	for ( int32 AssetIdx = 0; AssetIdx < BlueprintList.Num(); ++AssetIdx )
 	{
 		TSharedPtr< FClassViewerNode > NewNode;
 		LoadUnloadedTagData(NewNode, BlueprintList[AssetIdx]);
@@ -1893,17 +1927,7 @@ void FClassHierarchy::PopulateClassHierarchy()
 		BlueprintPackageToAssetDataMap.Add(BlueprintList[AssetIdx].PackageName, BlueprintList[AssetIdx]);
 	}
 
-	// Retrieve the children
-	FEditorClassHierarchy* EditorClassHierarchy = GEditor->EditorClassHierarchy;
-	check(EditorClassHierarchy);
-
-	int32 RootClassIndex;
-	RootClassIndex = EditorClassHierarchy->Find(TEXT("Object"));
-
-	ObjectClassRoot = MakeShareable(new FClassViewerNode(EditorClassHierarchy->GetClassName(RootClassIndex), false));
-	ObjectClassRoot->Class = EditorClassHierarchy->GetClass(RootClassIndex);
-
-	AddChildren_NoFilter(ObjectClassRoot, RootClassIndex, BlueprintPackageToAssetDataMap);
+	AddChildren_NoFilter(ObjectClassRoot, BlueprintPackageToAssetDataMap);
 
 	RootLevelClasses.Add(ObjectClassRoot);
 
diff --git a/Engine/Source/Editor/UnrealEd/Classes/Editor/EditorEngine.h b/Engine/Source/Editor/UnrealEd/Classes/Editor/EditorEngine.h
--- a/Engine/Source/Editor/UnrealEd/Classes/Editor/EditorEngine.h
+++ b/Engine/Source/Editor/UnrealEd/Classes/Editor/EditorEngine.h
@@ -433,9 +433,6 @@ public:
 	/** List of level editor viewport clients for level specific actions */
 	TArray<class FLevelEditorViewportClient*> LevelViewportClients;
 
-	/** Stores the class hierarchy generated from the make commandlet*/
-	class FEditorClassHierarchy*			EditorClassHierarchy;
-
 	/** Annotation to track which PIE/SIE (PlayWorld) UObjects have counterparts in the EditorWorld **/
 	class FUObjectAnnotationSparseBool ObjectsThatExistInEditorWorld;
 
diff --git a/Engine/Source/Editor/UnrealEd/Private/Editor.cpp b/Engine/Source/Editor/UnrealEd/Private/Editor.cpp
--- a/Engine/Source/Editor/UnrealEd/Private/Editor.cpp
+++ b/Engine/Source/Editor/UnrealEd/Private/Editor.cpp
@@ -345,21 +345,21 @@ void UEditorEngine::InitEditor(IEngineLoop* InEngineLoop)
 	// Call base.
 	UEngine::Init(InEngineLoop);
 
-    if( !FEngineBuildSettings::IsInternalBuild() &&
+	if( !FEngineBuildSettings::IsInternalBuild() &&
 		!FEngineBuildSettings::IsPerforceBuild() && 
 		!FApp::IsBenchmarking() &&
 		!GIsDemoMode && 
 		!IsRunningCommandlet() &&
 		!FPlatformProcess::IsApplicationRunning(TEXT("UnrealEngineLauncher") ) &&
 		!FPlatformProcess::IsApplicationRunning(TEXT("Unreal Engine Launcher") ) )
-    {
-        IDesktopPlatform* DesktopPlatform = FDesktopPlatformModule::Get();
-        if( DesktopPlatform != NULL )
-        {
-            DesktopPlatform->OpenLauncher( false, TEXT("") );
-        }
-    }
-    
+	{
+		IDesktopPlatform* DesktopPlatform = FDesktopPlatformModule::Get();
+		if( DesktopPlatform != NULL )
+		{
+			DesktopPlatform->OpenLauncher( false, TEXT("") );
+		}
+	}
+	
 	// Create selection sets.
 	PrivateInitSelectedSets();
 	
@@ -472,10 +472,6 @@ void UEditorEngine::Init(IEngineLoop* InEngineLoop)
 	GEditor = this;
 	InitEditor(InEngineLoop);
 
-	//Init the class hierarchy
-	EditorClassHierarchy = new FEditorClassHierarchy;
-	EditorClassHierarchy->Init();
-
 	Layers = FLayers::Create( TWeakObjectPtr< UEditorEngine >( this ) );
 
 	// Init transactioning.
@@ -569,9 +565,6 @@ void UEditorEngine::Init(IEngineLoop* InEngineLoop)
 
 	GLog->EnableBacklog( false );
 
-	// Re-init the class hierarchy to catch new classes due to module loads
-	EditorClassHierarchy->Init();
-
 	{
 		// avoid doing this every time, create a list of classes that derive from AVolume
 		TArray<UClass*> VolumeClasses;
diff --git a/Engine/Source/Editor/UnrealEd/Private/EditorClassHierarchy.cpp b/Engine/Source/Editor/UnrealEd/Private/EditorClassHierarchy.cpp
deleted file mode 100644
diff --git a/Engine/Source/Editor/UnrealEd/Public/Editor.h b/Engine/Source/Editor/UnrealEd/Public/Editor.h
--- a/Engine/Source/Editor/UnrealEd/Public/Editor.h
+++ b/Engine/Source/Editor/UnrealEd/Public/Editor.h
@@ -10,7 +10,6 @@
 
 #include "EngineInterpolationClasses.h"
 #include "Commandlets/EditorCommandlets.h"
-#include "EditorClassHierarchy.h"
 #include "EditorUndoClient.h"
 #include "TickableEditorObject.h"
 
@@ -87,7 +86,7 @@ struct UNREALED_API FEditorDelegates
 	DECLARE_MULTICAST_DELEGATE_OneParam(FOnBlueprintContextMenuCreated, FBlueprintGraphActionListBuilder &/*ContextMenuBuilder*/);
 
 	/** Called when the CurrentLevel is switched to a new level.  Note that this event won't be fired for temporary
-	    changes to the current level, such as when copying/pasting actors. */
+		changes to the current level, such as when copying/pasting actors. */
 	static FSimpleMulticastDelegate NewCurrentLevel;
 	/** Called when the map has changed */
 	static FOnMapChanged MapChange;
@@ -197,11 +196,11 @@ DECLARE_SCOPED_DELEGATE( RefreshAllBrowsers, FEditorDelegates::RefreshAllBrowser
 /** Texture alignment. */
 enum ETAxis
 {
-    TAXIS_X                 = 0,
-    TAXIS_Y                 = 1,
-    TAXIS_Z                 = 2,
-    TAXIS_WALLS             = 3,
-    TAXIS_AUTO              = 4,
+	TAXIS_X                 = 0,
+	TAXIS_Y                 = 1,
+	TAXIS_Z                 = 2,
+	TAXIS_WALLS             = 3,
+	TAXIS_AUTO              = 4,
 };
 
 /** Coordinate system identifiers. */
@@ -298,7 +297,7 @@ struct FImportObjectParams
 	UStruct*			ObjectStruct;
 
 	/** the original object that ImportObjectProperties was called for.
-	    if SubobjectOuter is a subobject, corresponds to the first object in SubobjectOuter's Outer chain that is not a subobject itself.
+		if SubobjectOuter is a subobject, corresponds to the first object in SubobjectOuter's Outer chain that is not a subobject itself.
 		if SubobjectOuter is not a subobject, should normally be the same value as SubobjectOuter */
 	UObject*			SubobjectRoot;
 
@@ -320,7 +319,7 @@ struct FImportObjectParams
 
 
 	/** True if we should call PreEditChange/PostEditChange on the object as it's imported.  Pass false here
-	    if you're going to do that on your own. */
+		if you're going to do that on your own. */
 	bool				bShouldCallEditChange;
 
 
diff --git a/Engine/Source/Editor/UnrealEd/Public/EditorClassHierarchy.h b/Engine/Source/Editor/UnrealEd/Public/EditorClassHierarchy.h
deleted file mode 100644
