- Fix class picker showing skeleton classes

diff --git a/Engine/Source/Editor/ClassViewer/Private/SClassViewer.cpp b/Engine/Source/Editor/ClassViewer/Private/SClassViewer.cpp
--- a/Engine/Source/Editor/ClassViewer/Private/SClassViewer.cpp
+++ b/Engine/Source/Editor/ClassViewer/Private/SClassViewer.cpp
@@ -1603,7 +1603,7 @@ void FClassHierarchy::AddChildren_NoFilter( TSharedPtr< FClassViewerNode >& InOu
 		UClass* CurrentClass = *ClassIt;
 
 		// Ignore deprecated and temporary trash classes.
-		if ( CurrentClass->HasAnyClassFlags(CLASS_Deprecated | CLASS_NewerVersionExists) )
+		if (CurrentClass->HasAnyClassFlags(CLASS_Deprecated | CLASS_NewerVersionExists) || FKismetEditorUtilities::IsClassABlueprintSkeleton(CurrentClass))
 		{
 			continue;
 		}
