DefaultPhysicsVolume no longer appears in the convert actor list in the details panel for Brushes/Volumes.  No longer allowed to drag/drop a Brush actor from the Class Viewer to the level viewport.
#ttp 338336 - CRITICAL: EDITOR: Convert Actor: Crash: Selecting DefaultPhysicsVolume as type
#branch UE4
#proj Editor.ClassViewer, Editor.UnrealEd
#remove Removed useless code in UEditorEngine::AddActor which asserted due to no UModel being associated with the ABrush CDO.
#add Made the condition more specific as to whether Brush type classes are added to the list in ClassViewer::Helpers::AddChildren_Tree and AddChildren_List, so that notplaceable ones are omitted.

diff --git a/Engine/Source/Editor/ClassViewer/Private/SClassViewer.cpp b/Engine/Source/Editor/ClassViewer/Private/SClassViewer.cpp
--- a/Engine/Source/Editor/ClassViewer/Private/SClassViewer.cpp
+++ b/Engine/Source/Editor/ClassViewer/Private/SClassViewer.cpp
@@ -589,7 +589,7 @@ namespace ClassViewer
 			bool bPassesPlaceableFilter = false;
 
 			// When in picker mode, brushes are valid "placeable" actors
-			if( bInOnlyPlaceables && InInitOptions.Mode == EClassViewerMode::ClassPicker && IsBrush(InOriginalRootNode->Class.Get()) )
+			if (bInOnlyPlaceables && InInitOptions.Mode == EClassViewerMode::ClassPicker && IsBrush(InOriginalRootNode->Class.Get()) && IsPlaceable(InOriginalRootNode->Class.Get()))
 			{
 				bPassesPlaceableFilter = true;
 			}
@@ -685,7 +685,7 @@ namespace ClassViewer
 			bool bPassesPlaceableFilter = false;
 
 			// When in picker mode, brushes are valid "placeable" actors
-			if( bInOnlyPlaceables && InInitOptions.Mode == EClassViewerMode::ClassPicker && IsBrush(InOriginalRootNode->Class.Get()) )
+			if( bInOnlyPlaceables && InInitOptions.Mode == EClassViewerMode::ClassPicker && IsBrush(InOriginalRootNode->Class.Get()) && IsPlaceable(InOriginalRootNode->Class.Get()))
 			{
 				bPassesPlaceableFilter = true;
 			}
diff --git a/Engine/Source/Editor/UnrealEd/Private/Editor.cpp b/Engine/Source/Editor/UnrealEd/Private/Editor.cpp
--- a/Engine/Source/Editor/UnrealEd/Private/Editor.cpp
+++ b/Engine/Source/Editor/UnrealEd/Private/Editor.cpp
@@ -6069,11 +6069,6 @@ AActor* UEditorEngine::AddActor(ULevel* InLevel, UClass* Class, const FVector& L
 		if( Actor )
 		{
 			SelectActor( Actor, 1, 0 );
-			ABrush* Brush = Cast<ABrush>(Actor);
-			if( Brush && Brush->BrushComponent )
-			{
-				FBSPOps::csgCopyBrush( Brush, Class->GetDefaultObject<ABrush>(), 0, RF_NoFlags, 1, true );
-			}
 			Actor->InvalidateLightingCache();
 			Actor->PostEditMove( true );
 		}
