Provide documentation tooltips for classes

diff --git a/Engine/Source/Editor/ClassViewer/Private/SClassViewer.cpp b/Engine/Source/Editor/ClassViewer/Private/SClassViewer.cpp
--- a/Engine/Source/Editor/ClassViewer/Private/SClassViewer.cpp
+++ b/Engine/Source/Editor/ClassViewer/Private/SClassViewer.cpp
@@ -21,6 +21,8 @@
 
 #include "UnloadedBlueprintData.h"
 
+#include "IDocumentation.h"
+
 #include "PropertyEditorModule.h"
 #include "PropertyHandle.h"
 
@@ -1321,24 +1323,26 @@ public:
 		
 		struct Local
 		{
-			static FText GetToolTipText(TSharedPtr<FClassViewerNode> AssociatedNode)
+
+			static TSharedPtr<SToolTip> GetToolTip(TSharedPtr<FClassViewerNode> AssociatedNode)
 			{
-				FText ToolTipText;
+				TSharedPtr<SToolTip> ToolTip;
 				if( AssociatedNode->PropertyHandle.IsValid() && AssociatedNode->IsRestricted() )
 				{
 					FText RestrictionToolTip;
 					AssociatedNode->PropertyHandle->GenerateRestrictionToolTip(*AssociatedNode->GetClassName(),RestrictionToolTip);
-					ToolTipText = RestrictionToolTip;
+
+					ToolTip = IDocumentation::Get()->CreateToolTip(RestrictionToolTip, nullptr, "", "");
 				}
 				else if (UClass* Class = AssociatedNode->Class.Get())
 				{
 					UPackage*  Package  = Class->GetOutermost();
 					UMetaData* MetaData = Package->GetMetaData();
 
-					ToolTipText = Class->GetToolTipText();
+					ToolTip = IDocumentation::Get()->CreateToolTip(Class->GetToolTipText(), nullptr, "Shared/Classes", Class->GetName());
 				}
 
-				return ToolTipText;
+				return ToolTip;
 			}
 		};
 
@@ -1374,7 +1378,7 @@ public:
 						.Text( *ClassName.Get() )
 						.HighlightText(*InArgs._HighlightText)
 						.ColorAndOpacity( this, &SClassItem::GetTextColor)
-						.ToolTipText_Static(&Local::GetToolTipText, AssociatedNode)
+						.ToolTip(Local::GetToolTip(AssociatedNode))
 						.IsEnabled(!bIsRestricted)
 				]
 
