Fixed TreeView hierarchy in Class Viewer when changing filters.
#ttp 338730 - EDITOR: Class Viewer: TreeView hierarchy in the Class Viewer can get broken
#branch UE4
#proj Editor.ClassViewer
#add Added an SClassViewer::bPendingSetExpansionStates member whose purpose is to ensure that the tree view expansion state is synched with the SClassViewer copy the following tick.  This is necessary in order to update child's expansion state when it is changed in the OnGetChildren callback (because otherwise the STreeView will clear the pending refresh flags before they have a chance to be dealt with).  Now OnGetChildren just sets this flag and forces a sync of the expansion state on the next tick, via SetExpansionStatesInTree.
#add Added an OnExpansionChanged callback to allow the local expansion state cache to remain in sync with that of the STreeView.
#add ExpandRootNodes now sets an entry in the ExpansionStateMap (as it should), so that SetExpansionStatesInTree will recurse beyond the root nodes correctly.

diff --git a/Engine/Source/Editor/ClassViewer/Private/SClassViewer.cpp b/Engine/Source/Editor/ClassViewer/Private/SClassViewer.cpp
--- a/Engine/Source/Editor/ClassViewer/Private/SClassViewer.cpp
+++ b/Engine/Source/Editor/ClassViewer/Private/SClassViewer.cpp
@@ -2028,6 +2028,7 @@ void SClassViewer::Construct(const FArguments& InArgs, const FClassViewerInitial
 	OnClassPicked = InArgs._OnClassPickedDelegate;
 
 	bSaveExpansionStates = true;
+	bPendingSetExpansionStates = false;
 
 	bEnableClassDynamicLoading = InInitOptions.bEnableClassDynamicLoading;
 
@@ -2154,6 +2155,9 @@ void SClassViewer::Construct(const FArguments& InArgs, const FClassViewerInitial
 					// Find out when the user selects something in the tree
 					.OnSelectionChanged( this, &SClassViewer::OnClassViewerSelectionChanged )
 
+					// Called when the expansion state of an item changes
+					.OnExpansionChanged( this, &SClassViewer::OnClassViewerExpansionChanged )
+
 					// Allow for some spacing between items with a larger item height.
 					.ItemHeight(20.0f)
 
@@ -2286,6 +2290,17 @@ void SClassViewer::OnClassViewerSelectionChanged( TSharedPtr<FClassViewerNode> I
 	}
 }
 
+void SClassViewer::OnClassViewerExpansionChanged(TSharedPtr<FClassViewerNode> Item, bool bExpanded)
+{
+	// Sometimes the item is not valid anymore due to filtering.
+	if (Item.IsValid() == false || Item->IsRestricted())
+	{
+		return;
+	}
+
+	ExpansionStateMap.Add(Item->GetClassName(), bExpanded);
+}
+
 TSharedPtr< SWidget > SClassViewer::BuildMenuWidget()
 {
 	// Empty list of commands.
@@ -2410,7 +2425,7 @@ TSharedRef< ITableRow > SClassViewer::OnGenerateRowForClassViewer( TSharedPtr<FC
 	bool* bIsExpanded = ExpansionStateMap.Find( Item->GetClassName() );
 	if( bIsExpanded && *bIsExpanded )
 	{
-		ClassTree->SetItemExpansion( Item, *bIsExpanded );
+		bPendingSetExpansionStates = true;
 	}
 
 	return ReturnRow;
@@ -2430,6 +2445,7 @@ void SClassViewer::ExpandRootNodes()
 {
 	for (int32 NodeIdx = 0; NodeIdx < RootTreeItems.Num(); ++NodeIdx)
 	{
+		ExpansionStateMap.Add(RootTreeItems[NodeIdx]->GetClassName(), true);
 		ClassTree->SetItemExpansion(RootTreeItems[NodeIdx], true);
 	}
 }
@@ -2913,6 +2929,12 @@ void SClassViewer::Tick( const FGeometry& AllottedGeometry, const double InCurre
 			ExpandRootNodes();
 		}
 	}
+
+	if (bPendingSetExpansionStates)
+	{
+		bPendingSetExpansionStates = false;
+		SetExpansionStatesInTree(RootTreeItems[0]);
+	}
 }
 
 #undef LOCTEXT_NAMESPACE
diff --git a/Engine/Source/Editor/ClassViewer/Private/SClassViewer.h b/Engine/Source/Editor/ClassViewer/Private/SClassViewer.h
--- a/Engine/Source/Editor/ClassViewer/Private/SClassViewer.h
+++ b/Engine/Source/Editor/ClassViewer/Private/SClassViewer.h
@@ -65,6 +65,9 @@ private:
 	/** Called by Slate when an item is selected from the tree/list. */
 	void OnClassViewerSelectionChanged( TSharedPtr<FClassViewerNode> Item, ESelectInfo::Type SelectInfo );
 
+	/** Called by Slate when an item is expanded/collapsed from the tree/list. */
+	void OnClassViewerExpansionChanged(TSharedPtr<FClassViewerNode> Item, bool bExpanded);
+
 	/** A simple function for the filters to keep them enabled. */
 	bool Menu_CanExecute() const;
 
@@ -212,4 +215,7 @@ private:
 
 	/** True if the search box will take keyboard focus next frame */
 	bool bPendingFocusNextFrame;
+
+	/** True if we need to set the tree expansion states according to our local copy next tick */
+	bool bPendingSetExpansionStates;
 };
