Enabled dynamic class loading in the Class picker dialog and fixed the Class viewer container expansion state tracking.
#Branch UE4
#Proj UnrealEd, ClassViewer

- enabling dynamic class loading requested by [REDACTED_PERSON] after encountering usability issues with the class picker dialog
- the change to the class viewer expansion map was an issue encountered once the dynamic class loading was enabled, the TSharedPtr<FString> key was hashing the container and not the string value which resulted in expansion state restoration failing and a minor bloating of the map class.

diff --git a/Engine/Source/Editor/ClassViewer/Private/SClassViewer.cpp b/Engine/Source/Editor/ClassViewer/Private/SClassViewer.cpp
--- a/Engine/Source/Editor/ClassViewer/Private/SClassViewer.cpp
+++ b/Engine/Source/Editor/ClassViewer/Private/SClassViewer.cpp
@@ -2298,7 +2298,7 @@ void SClassViewer::OnClassViewerExpansionChanged(TSharedPtr<FClassViewerNode> It
 		return;
 	}
 
-	ExpansionStateMap.Add(Item->GetClassName(), bExpanded);
+	ExpansionStateMap.Add(*(Item->GetClassName()), bExpanded);
 }
 
 TSharedPtr< SWidget > SClassViewer::BuildMenuWidget()
@@ -2422,7 +2422,7 @@ TSharedRef< ITableRow > SClassViewer::OnGenerateRowForClassViewer( TSharedPtr<FC
 	}
 
 	// Expand the item if needed.
-	bool* bIsExpanded = ExpansionStateMap.Find( Item->GetClassName() );
+	bool* bIsExpanded = ExpansionStateMap.Find( *(Item->GetClassName()) );
 	if( bIsExpanded && *bIsExpanded )
 	{
 		bPendingSetExpansionStates = true;
@@ -2445,7 +2445,7 @@ void SClassViewer::ExpandRootNodes()
 {
 	for (int32 NodeIdx = 0; NodeIdx < RootTreeItems.Num(); ++NodeIdx)
 	{
-		ExpansionStateMap.Add(RootTreeItems[NodeIdx]->GetClassName(), true);
+		ExpansionStateMap.Add(*(RootTreeItems[NodeIdx]->GetClassName()), true);
 		ClassTree->SetItemExpansion(RootTreeItems[NodeIdx], true);
 	}
 }
@@ -2712,7 +2712,7 @@ bool SClassViewer::ExpandFilteredInNodes(TSharedPtr<FClassViewerNode> InNode)
 
 void SClassViewer::MapExpansionStatesInTree( TSharedPtr<FClassViewerNode> InItem )
 {
-	ExpansionStateMap.Add( InItem->GetClassName(), ClassTree->IsItemExpanded( InItem ) );
+	ExpansionStateMap.Add( *(InItem->GetClassName()), ClassTree->IsItemExpanded( InItem ) );
 
 	// Map out all the children, this will be done recursively.
 	for( int32 ChildIdx(0); ChildIdx < InItem->GetChildrenList().Num(); ++ChildIdx )
@@ -2723,7 +2723,7 @@ void SClassViewer::MapExpansionStatesInTree( TSharedPtr<FClassViewerNode> InItem
 
 void SClassViewer::SetExpansionStatesInTree( TSharedPtr<FClassViewerNode> InItem )
 {
-	bool* bIsExpanded = ExpansionStateMap.Find( InItem->GetClassName() );
+	bool* bIsExpanded = ExpansionStateMap.Find( *(InItem->GetClassName()) );
 	if( bIsExpanded )
 	{
 		ClassTree->SetItemExpansion( InItem, *bIsExpanded );
@@ -2754,7 +2754,7 @@ void SClassViewer::Populate()
 			for( int32 ChildIdx(0); ChildIdx < RootTreeItems.Num(); ++ChildIdx )
 			{
 				// Check if the item is actually expanded or if it's only expanded because it is root level.
-				bool* bIsExpanded = ExpansionStateMap.Find( RootTreeItems[ChildIdx]->GetClassName() );
+				bool* bIsExpanded = ExpansionStateMap.Find( *(RootTreeItems[ChildIdx]->GetClassName()) );
 				if(bIsExpanded && !*bIsExpanded || !bIsExpanded)
 				{
 					ClassTree->SetItemExpansion( RootTreeItems[ChildIdx], false );
diff --git a/Engine/Source/Editor/ClassViewer/Private/SClassViewer.h b/Engine/Source/Editor/ClassViewer/Private/SClassViewer.h
--- a/Engine/Source/Editor/ClassViewer/Private/SClassViewer.h
+++ b/Engine/Source/Editor/ClassViewer/Private/SClassViewer.h
@@ -208,7 +208,7 @@ private:
 	bool bSaveExpansionStates;
 
 	/** The map holding the expansion state map for the tree. */
-	TMap< TSharedPtr< FString >, bool > ExpansionStateMap;
+	TMap< FString, bool > ExpansionStateMap;
 
 	/** True if the Class Viewer needs to be repopulated at the next appropriate opportunity, occurs whenever classes are added, removed, renamed, etc. */
 	bool bNeedsRefresh;
diff --git a/Engine/Source/Editor/UnrealEd/Private/Factories/EditorFactories.cpp b/Engine/Source/Editor/UnrealEd/Private/Factories/EditorFactories.cpp
--- a/Engine/Source/Editor/UnrealEd/Private/Factories/EditorFactories.cpp
+++ b/Engine/Source/Editor/UnrealEd/Private/Factories/EditorFactories.cpp
@@ -6172,7 +6172,7 @@ bool UBlueprintFactory::ConfigureProperties()
 	Options.bShowUnloadedBlueprints = true;
 
 	// Enable Class Dynamic Loading
-	Options.bEnableClassDynamicLoading = false;
+	Options.bEnableClassDynamicLoading = true;
 
 	// Prevent creating blueprints of classes that require special setup (they'll be allowed in the corresponding factories / via other means)
 	TSharedPtr<FBlueprintParentFilter> Filter = MakeShareable(new FBlueprintParentFilter);
