Class viewer no longer crashes when it finds zero results.

#ttp 341701 - MAIN: EDITOR: CRASH: Editor crashes when searching for a nonexistant reparent class

#branch UE4

#change Very timing specific, the crash actually occurs (in the example) between pressing 's' then 'd' as the 's' search sets bPendingSetExpansionStates to true, then before it can refresh the list the search is updated with 'd' which will result in zero matches, however the bPendingSetExpansionStates flag is still set to true so it tries to expand something that isn't there. This change just resets the flag whenever the list is refreshed.

#change Also modified the code so it doesn't do the 'Find' if it doesn't need to, and threw in a check incase this happens again.

diff --git a/Engine/Source/Editor/ClassViewer/Private/SClassViewer.cpp b/Engine/Source/Editor/ClassViewer/Private/SClassViewer.cpp
--- a/Engine/Source/Editor/ClassViewer/Private/SClassViewer.cpp
+++ b/Engine/Source/Editor/ClassViewer/Private/SClassViewer.cpp
@@ -2423,10 +2423,13 @@ TSharedRef< ITableRow > SClassViewer::OnGenerateRowForClassViewer( TSharedPtr<FC
 	}
 
 	// Expand the item if needed.
-	bool* bIsExpanded = ExpansionStateMap.Find( *(Item->GetClassName()) );
-	if( bIsExpanded && *bIsExpanded )
+	if (!bPendingSetExpansionStates)
 	{
-		bPendingSetExpansionStates = true;
+		bool* bIsExpanded = ExpansionStateMap.Find(*(Item->GetClassName()));
+		if (bIsExpanded && *bIsExpanded)
+		{
+			bPendingSetExpansionStates = true;
+		}
 	}
 
 	return ReturnRow;
@@ -2747,6 +2750,8 @@ void SClassViewer::SetExpansionStatesInTree( TSharedPtr<FClassViewerNode> InItem
 
 void SClassViewer::Populate()
 {
+	bPendingSetExpansionStates = false;
+
 	// If showing a class tree, we may need to save expansion states.
 	if( ClassTree->GetVisibility() == EVisibility::Visible )
 	{
@@ -2933,8 +2938,9 @@ void SClassViewer::Tick( const FGeometry& AllottedGeometry, const double InCurre
 
 	if (bPendingSetExpansionStates)
 	{
-		bPendingSetExpansionStates = false;
+		check(RootTreeItems.Num() > 0);
 		SetExpansionStatesInTree(RootTreeItems[0]);
+		bPendingSetExpansionStates = false;
 	}
 }
 
