Only Actors should be considered placeable

diff --git a/Engine/Source/Editor/ClassViewer/Private/SClassViewer.cpp b/Engine/Source/Editor/ClassViewer/Private/SClassViewer.cpp
--- a/Engine/Source/Editor/ClassViewer/Private/SClassViewer.cpp
+++ b/Engine/Source/Editor/ClassViewer/Private/SClassViewer.cpp
@@ -483,7 +483,7 @@ namespace ClassViewer
 		 */
 		static bool IsPlaceable(UClass* InClass)
 		{
-			return !InClass->HasAnyClassFlags(CLASS_Abstract | CLASS_NotPlaceable);
+			return !InClass->HasAnyClassFlags(CLASS_Abstract | CLASS_NotPlaceable) && InClass->IsChildOf<AActor>();
 		}
 
 		/** Checks if a particular class is abstract.
