Class Viewer gives actual information (about Placeable flag) for unloaded assets.
OnlyPlaceable metadata added. Indicates whether only placeable classes should be shown in the class picker.
More robust AGameMode::PreInitializeComponents

ttp345655 CRITICAL: CRASH: BLUEPRINT: adding GameMode blueprint as a child actor

diff --git a/Engine/Source/Editor/ClassViewer/Private/ClassViewerNode.cpp b/Engine/Source/Editor/ClassViewer/Private/ClassViewerNode.cpp
--- a/Engine/Source/Editor/ClassViewer/Private/ClassViewerNode.cpp
+++ b/Engine/Source/Editor/ClassViewer/Private/ClassViewerNode.cpp
@@ -7,12 +7,11 @@
 #include "PropertyEditorModule.h"
 #include "PropertyHandle.h"
 
-FClassViewerNode::FClassViewerNode( const FString& InClassName, const FString& InClassDisplayName, bool bInIsPlaceable )
+FClassViewerNode::FClassViewerNode( const FString& InClassName, const FString& InClassDisplayName )
 {
 	ClassName = MakeShareable(new FString(InClassName));
 	ClassDisplayName = MakeShareable(new FString(InClassDisplayName));
 	bPassesFilter = false;
-	bIsClassPlaceable = bInIsPlaceable;
 	bIsBPNormalType = false;
 
 	Class = NULL;
@@ -24,7 +23,6 @@ FClassViewerNode::FClassViewerNode( const FClassViewerNode& InCopyObject)
 	ClassName = InCopyObject.ClassName;
 	ClassDisplayName = InCopyObject.ClassDisplayName;
 	bPassesFilter = InCopyObject.bPassesFilter;
-	bIsClassPlaceable = InCopyObject.bIsClassPlaceable;
 
 	Class = InCopyObject.Class;
 	Blueprint = InCopyObject.Blueprint;
@@ -87,4 +85,26 @@ void FClassViewerNode::AddUniqueChild(TSharedPtr<FClassViewerNode> NewChild)
 bool FClassViewerNode::IsRestricted() const
 {
 	return PropertyHandle.IsValid() && PropertyHandle->IsRestricted(*ClassName);
-}
\ No newline at end of file
+}
+
+bool FClassViewerNode::IsClassPlaceable() const
+{
+	const UClass* LoadedClass = Class.Get();
+	if (LoadedClass)
+	{
+		const bool bPlaceableFlags = !LoadedClass->HasAnyClassFlags(CLASS_Abstract | CLASS_NotPlaceable);
+		const bool bBasedOnActor = LoadedClass->IsChildOf(AActor::StaticClass());
+		const bool bNotABrush = !LoadedClass->IsChildOf(ABrush::StaticClass());
+		return bPlaceableFlags && bBasedOnActor && bNotABrush;
+	}
+
+	if (UnloadedBlueprintData.IsValid())
+	{
+		const bool bPlaceableFlags = !UnloadedBlueprintData->HasAnyClassFlags(CLASS_Abstract | CLASS_NotPlaceable);
+		const bool bBasedOnActor = UnloadedBlueprintData->IsChildOf(AActor::StaticClass());
+		const bool bNotABrush = !UnloadedBlueprintData->IsChildOf(ABrush::StaticClass());
+		return bPlaceableFlags && bBasedOnActor && bNotABrush;
+	}
+
+	return false;
+}
diff --git a/Engine/Source/Editor/ClassViewer/Private/ClassViewerNode.h b/Engine/Source/Editor/ClassViewer/Private/ClassViewerNode.h
--- a/Engine/Source/Editor/ClassViewer/Private/ClassViewerNode.h
+++ b/Engine/Source/Editor/ClassViewer/Private/ClassViewerNode.h
@@ -12,7 +12,7 @@ public:
 	 * @param	InClassDisplayName				The display name of the class this node represents
 	 * @param	bInIsPlaceable					true if the class is a placeable class.
 	 */
-	FClassViewerNode( const FString& InClassName, const FString& InClassDisplayName, bool bInIsPlaceable );
+	FClassViewerNode( const FString& InClassName, const FString& InClassDisplayName );
 
 	FClassViewerNode( const FClassViewerNode& InCopyObject);
 
@@ -47,15 +47,7 @@ public:
 	}
 
 	/** Checks if the class is placeable. */
-	bool IsClassPlaceable() const
-	{
-		return bIsClassPlaceable;
-	}
-
-	void IsClassPlaceable(bool bInIsClassPlaceable)
-	{
-		bIsClassPlaceable = bInIsClassPlaceable;
-	}
+	bool IsClassPlaceable() const;
 
 	bool IsRestricted() const;
 
@@ -69,9 +61,6 @@ private:
 	/** List of children. */
 	TArray<TSharedPtr<FClassViewerNode>> ChildrenList;
 
-	/** true if the class is placeable. */
-	bool bIsClassPlaceable;
-
 public:
 	/** The class this node is associated with. */
 	TWeakObjectPtr<UClass> Class;
diff --git a/Engine/Source/Editor/ClassViewer/Private/SClassViewer.cpp b/Engine/Source/Editor/ClassViewer/Private/SClassViewer.cpp
--- a/Engine/Source/Editor/ClassViewer/Private/SClassViewer.cpp
+++ b/Engine/Source/Editor/ClassViewer/Private/SClassViewer.cpp
@@ -476,7 +476,7 @@ namespace ClassViewer
 		 *	@param InClass				The Class to check.
 		 *	@return Returns true if the class is a brush.
 		 */
-		static bool IsBrush(UClass* InClass)
+		static bool IsBrush(const UClass* InClass)
 		{
 			return InClass->IsChildOf( ABrush::StaticClass() );
 		}
@@ -485,7 +485,7 @@ namespace ClassViewer
 		 *	@param InClass				The Class to check.
 		 *	@return Returns true if the class is placeable.
 		 */
-		static bool IsPlaceable(UClass* InClass)
+		static bool IsPlaceable(const UClass* InClass)
 		{
 			return !InClass->HasAnyClassFlags(CLASS_Abstract | CLASS_NotPlaceable) && InClass->IsChildOf<AActor>();
 		}
@@ -494,7 +494,7 @@ namespace ClassViewer
 		 *	@param InClass				The Class to check.
 		 *	@return Returns true if the class is abstract.
 		 */
-		static bool IsAbstract(UClass* InClass)
+		static bool IsAbstract(const UClass* InClass)
 		{
 			return InClass->HasAnyClassFlags(CLASS_Abstract);
 		}
@@ -1623,12 +1623,8 @@ FClassHierarchy::~FClassHierarchy()
 
 static TSharedPtr< FClassViewerNode > CreateNodeForClass(UClass* Class, const TMultiMap<FName, FAssetData>& BlueprintPackageToAssetDataMap)
 {
-	const bool bIsPlaceable = ClassViewer::Helpers::IsPlaceable(Class);
-	const bool bIsAbstract = ClassViewer::Helpers::IsAbstract(Class);
-	const bool bIsBrush = ClassViewer::Helpers::IsBrush(Class);
-
 	// Create the new node so it can be passed to AddChildren, fill it in with if it is placeable, abstract, and/or a brush.
-	TSharedPtr< FClassViewerNode > NewNode = MakeShareable(new FClassViewerNode(Class->GetName(), Class->GetDisplayNameText().ToString(), bIsPlaceable && !bIsAbstract && !bIsBrush));
+	TSharedPtr< FClassViewerNode > NewNode = MakeShareable(new FClassViewerNode(Class->GetName(), Class->GetDisplayNameText().ToString()));
 	NewNode->Blueprint = ClassViewer::Helpers::GetBlueprint(Class);
 	NewNode->Class = Class;
 
@@ -1660,7 +1656,7 @@ void FClassHierarchy::AddChildren_NoFilter( TSharedPtr< FClassViewerNode >& InOu
 {
 	UClass* RootClass = UObject::StaticClass();
 
-	ObjectClassRoot = MakeShareable(new FClassViewerNode(RootClass->GetName(), RootClass->GetDisplayNameText().ToString(), false));
+	ObjectClassRoot = MakeShareable(new FClassViewerNode(RootClass->GetName(), RootClass->GetDisplayNameText().ToString()));
 	ObjectClassRoot->Class = RootClass;
 
 	TMap< UClass*, TSharedPtr< FClassViewerNode > > Nodes;
@@ -1914,7 +1910,7 @@ void FClassHierarchy::LoadUnloadedTagData(TSharedPtr<FClassViewerNode>& InOutCla
 	const FString AssetName = InAssetData.AssetName.ToString();
 
 	// Create the viewer node.
-	InOutClassViewerNode = MakeShareable(new FClassViewerNode(AssetName, AssetName, true));
+	InOutClassViewerNode = MakeShareable(new FClassViewerNode(AssetName, AssetName));
 			
 	// It is an unloaded blueprint, so we need to create the structure that will hold the data.
 	TSharedPtr<FUnloadedBlueprintData> UnloadedBlueprintData = MakeShareable( new FUnloadedBlueprintData(InOutClassViewerNode) );
@@ -2912,7 +2908,7 @@ FReply SClassViewer::OnKeyboardFocusReceived( const FGeometry& MyGeometry, const
 
 TSharedPtr<FClassViewerNode> SClassViewer::CreateNoneOption()
 {
-	TSharedPtr<FClassViewerNode> NoneItem = MakeShareable( new FClassViewerNode("None", "None", false) );
+	TSharedPtr<FClassViewerNode> NoneItem = MakeShareable( new FClassViewerNode("None", "None") );
 
 	// The item "passes" the filter so it does not appear grayed out.
 	NoneItem->bPassesFilter = true;
diff --git a/Engine/Source/Editor/PropertyEditor/Private/PropertyHandleImpl.cpp b/Engine/Source/Editor/PropertyEditor/Private/PropertyHandleImpl.cpp
--- a/Engine/Source/Editor/PropertyEditor/Private/PropertyHandleImpl.cpp
+++ b/Engine/Source/Editor/PropertyEditor/Private/PropertyHandleImpl.cpp
@@ -1929,16 +1929,21 @@ bool FPropertyHandleBase::GeneratePossibleValues(TArray< TSharedPtr<FString> >&
 
 		const bool bAllowAbstract = Property->GetOwnerProperty()->HasMetaData(TEXT("AllowAbstract"));
 		const bool bBlueprintBaseOnly = Property->GetOwnerProperty()->HasMetaData(TEXT("BlueprintBaseOnly"));
+		const bool bAllowOnlyPlaceable = Property->GetOwnerProperty()->HasMetaData(TEXT("OnlyPlaceable"));
 		UClass* InterfaceThatMustBeImplemented = Property->GetOwnerProperty()->GetClassMetaData(TEXT("MustImplement"));
 
-		for( TObjectIterator<UClass> It ; It ; ++It )
+		if (!bAllowOnlyPlaceable || MetaClass->IsChildOf<AActor>())
 		{
-			if ( It->IsChildOf(MetaClass) && 
-				PropertyEditorHelpers::IsEditInlineClassAllowed(*It, bAllowAbstract) &&
-				(!bBlueprintBaseOnly || FKismetEditorUtilities::CanCreateBlueprintOfClass(*It)) &&
-				(!InterfaceThatMustBeImplemented || It->ImplementsInterface(InterfaceThatMustBeImplemented)) )
+			for (TObjectIterator<UClass> It; It; ++It)
 			{
-				OutOptionStrings.Add( TSharedPtr< FString >( new FString( It->GetName() ) ) );
+				if (It->IsChildOf(MetaClass)
+					&& PropertyEditorHelpers::IsEditInlineClassAllowed(*It, bAllowAbstract)
+					&& (!bBlueprintBaseOnly || FKismetEditorUtilities::CanCreateBlueprintOfClass(*It))
+					&& (!InterfaceThatMustBeImplemented || It->ImplementsInterface(InterfaceThatMustBeImplemented))
+					&& (!bAllowOnlyPlaceable || !It->HasAnyClassFlags(CLASS_Abstract | CLASS_NotPlaceable)))
+				{
+					OutOptionStrings.Add(TSharedPtr< FString >(new FString(It->GetName())));
+				}
 			}
 		}
 	}
diff --git a/Engine/Source/Editor/PropertyEditor/Private/UserInterface/PropertyEditor/SPropertyEditorClass.cpp b/Engine/Source/Editor/PropertyEditor/Private/UserInterface/PropertyEditor/SPropertyEditorClass.cpp
--- a/Engine/Source/Editor/PropertyEditor/Private/UserInterface/PropertyEditor/SPropertyEditorClass.cpp
+++ b/Engine/Source/Editor/PropertyEditor/Private/UserInterface/PropertyEditor/SPropertyEditorClass.cpp
@@ -104,6 +104,7 @@ void SPropertyEditorClass::Construct(const FArguments& InArgs, const TSharedPtr<
 		}
 		
 		bAllowAbstract = Property->GetOwnerProperty()->HasMetaData(TEXT("AllowAbstract"));
+		bAllowOnlyPlaceable = Property->GetOwnerProperty()->HasMetaData(TEXT("OnlyPlaceable"));
 		bIsBlueprintBaseOnly = Property->GetOwnerProperty()->HasMetaData(TEXT("BlueprintBaseOnly"));
 		RequiredInterface = Property->GetOwnerProperty()->GetClassMetaData(TEXT("MustImplement"));
 		bAllowNone = !(Property->PropertyFlags & CPF_NoClear);
@@ -119,6 +120,7 @@ void SPropertyEditorClass::Construct(const FArguments& InArgs, const TSharedPtr<
 		bAllowAbstract = InArgs._AllowAbstract;
 		bIsBlueprintBaseOnly = InArgs._IsBlueprintBaseOnly;
 		bAllowNone = InArgs._AllowNone;
+		bAllowOnlyPlaceable = false;
 
 		SelectedClass = InArgs._SelectedClass;
 		OnSetClass = InArgs._OnSetClass;
@@ -191,6 +193,7 @@ TSharedRef<SWidget> SPropertyEditorClass::GenerateClassPicker()
 	ClassFilter->InterfaceThatMustBeImplemented = RequiredInterface;
 	ClassFilter->bAllowAbstract = bAllowAbstract;
 	Options.bIsBlueprintBaseOnly = bIsBlueprintBaseOnly;
+	Options.bIsPlaceableOnly = bAllowOnlyPlaceable;
 
 	FOnClassPicked OnPicked(FOnClassPicked::CreateRaw(this, &SPropertyEditorClass::OnClassPicked));
 
diff --git a/Engine/Source/Editor/PropertyEditor/Private/UserInterface/PropertyEditor/SPropertyEditorClass.h b/Engine/Source/Editor/PropertyEditor/Private/UserInterface/PropertyEditor/SPropertyEditorClass.h
--- a/Engine/Source/Editor/PropertyEditor/Private/UserInterface/PropertyEditor/SPropertyEditorClass.h
+++ b/Engine/Source/Editor/PropertyEditor/Private/UserInterface/PropertyEditor/SPropertyEditorClass.h
@@ -85,6 +85,8 @@ private:
 	bool bIsBlueprintBaseOnly;
 	/** Should we be able to select "None" as a class? */
 	bool bAllowNone;
+	/** Should only placeable classes be displayed? */
+	bool bAllowOnlyPlaceable;
 
 	/** Attribute used to get the currently selected class (required if PropertyEditor == null) */
 	TAttribute<const UClass*> SelectedClass;
diff --git a/Engine/Source/Runtime/CoreUObject/Public/UObject/ObjectBase.h b/Engine/Source/Runtime/CoreUObject/Public/UObject/ObjectBase.h
--- a/Engine/Source/Runtime/CoreUObject/Public/UObject/ObjectBase.h
+++ b/Engine/Source/Runtime/CoreUObject/Public/UObject/ObjectBase.h
@@ -939,6 +939,9 @@ namespace UM
 		/// [PropertyMetadata] Used for FStringClassReference properties.  Indicates whether only blueprint classes should be shown in the class picker.
 		IsBlueprintBaseOnly,
 
+		/// [PropertyMetadata] Used for Subclass properties. Indicates whether only placeable classes should be shown in the class picker.
+		OnlyPlaceable,
+
 		/// [PropertyMetadata]
 		MakeEditWidget,
 
diff --git a/Engine/Source/Runtime/Engine/Classes/Components/ChildActorComponent.h b/Engine/Source/Runtime/Engine/Classes/Components/ChildActorComponent.h
--- a/Engine/Source/Runtime/Engine/Classes/Components/ChildActorComponent.h
+++ b/Engine/Source/Runtime/Engine/Classes/Components/ChildActorComponent.h
@@ -13,7 +13,7 @@ class UChildActorComponent : public USceneComponent
 	GENERATED_UCLASS_BODY()
 
 	/** The class of Actor to spawn */
-	UPROPERTY(EditAnywhere, BlueprintReadOnly, Category=ChildActorComponent)
+	UPROPERTY(EditAnywhere, BlueprintReadOnly, Category=ChildActorComponent, meta=(OnlyPlaceable))
 	TSubclassOf<AActor>	ChildActorClass;
 
 	/** The actor that we spawned and own */
diff --git a/Engine/Source/Runtime/Engine/Private/GameMode.cpp b/Engine/Source/Runtime/Engine/Private/GameMode.cpp
--- a/Engine/Source/Runtime/Engine/Private/GameMode.cpp
+++ b/Engine/Source/Runtime/Engine/Private/GameMode.cpp
@@ -400,7 +400,10 @@ void AGameMode::PreInitializeComponents()
 
 	GameState = GetWorld()->SpawnActor<AGameState>(GameStateClass, SpawnInfo);
 	GetWorld()->GameState = GameState;
-	GameState->AuthorityGameMode = this;
+	if (GameState)
+	{
+		GameState->AuthorityGameMode = this;
+	}
 
 	// Only need NetworkManager for servers in net games
 	GetWorld()->NetworkManager = GetWorldSettings()->GameNetworkManagerClass ? GetWorld()->SpawnActor<AGameNetworkManager>(GetWorldSettings()->GameNetworkManagerClass, SpawnInfo ) : NULL;
