Class viewer will now display unloaded Blutility, UMG, and GameplayAbility Blueprints.

#ttp 348359 - UE4: Gameplay Ability Blueprint classes do not show up in the Class Viewer picker

diff --git a/Engine/Source/Editor/ClassViewer/Private/SClassViewer.cpp b/Engine/Source/Editor/ClassViewer/Private/SClassViewer.cpp
--- a/Engine/Source/Editor/ClassViewer/Private/SClassViewer.cpp
+++ b/Engine/Source/Editor/ClassViewer/Private/SClassViewer.cpp
@@ -1977,6 +1977,9 @@ void FClassHierarchy::PopulateClassHierarchy()
 	Filter.ClassNames.Add(UBlueprint::StaticClass()->GetFName());
 	Filter.ClassNames.Add(UAnimBlueprint::StaticClass()->GetFName());
 	Filter.ClassNames.Add(UBlueprintGeneratedClass::StaticClass()->GetFName());
+
+	// Include any Blueprint based objects as well, this includes things like Blutilities, UMG, and GameplayAbility objects
+	Filter.bRecursiveClasses = true;
 	AssetRegistryModule.Get().GetAssets(Filter, BlueprintList);
 
 	TMultiMap<FName,FAssetData> BlueprintPackageToAssetDataMap;
