Class Pickers will allow you to type to search and use the arrow keys to navigate without changing focus with the tab key or mouse.

Class Picker or Content Picker is not created when creating a pin of class/object type, instead it is generated when clicking to drop down the combo box.

diff --git a/Engine/Source/Editor/ClassViewer/Private/SClassViewer.cpp b/Engine/Source/Editor/ClassViewer/Private/SClassViewer.cpp
--- a/Engine/Source/Editor/ClassViewer/Private/SClassViewer.cpp
+++ b/Engine/Source/Editor/ClassViewer/Private/SClassViewer.cpp
@@ -33,6 +33,8 @@
 #include "HotReloadInterface.h"
 #include "SSearchBox.h"
 
+#include "SListViewSelectorDropdownMenu.h"
+
 #define LOCTEXT_NAMESPACE "SClassViewer"
 
 DEFINE_LOG_CATEGORY_STATIC(LogEditorClassViewer, Log, All);
@@ -1272,7 +1274,7 @@ namespace ClassViewer
 DECLARE_DELEGATE_OneParam( FOnClassItemDoubleClickDelegate, TSharedPtr<FClassViewerNode> );
 
 /** The item used for visualizing the class in the tree. */
-class SClassItem : public STableRow< TSharedPtr<FString> >
+class SClassItem : public SComboRow< TSharedPtr<FString> >
 {
 public:
 	
@@ -2097,71 +2099,105 @@ void SClassViewer::Construct(const FArguments& InArgs, const FClassViewerInitial
 		OnContextMenuOpening = FOnContextMenuOpening::CreateSP(this, &SClassViewer::BuildMenuWidget);
 	}
 
-	
-	this->ChildSlot
+	// Holds the bulk of the class viewer's sub-widgets, to be added to the widget after construction
+	TSharedPtr< SWidget > ClassViewerContent;
+
+	SAssignNew(ClassViewerContent, SVerticalBox)
+	+SVerticalBox::Slot()
+	.AutoHeight()
+	[
+		MenuBarBuilder.MakeWidget()
+	]
+
+	+SVerticalBox::Slot()
+	.AutoHeight()
+	.Padding( 1.0f, 0.0f, 1.0f, 0.0f )
 	[
-		SNew(SVerticalBox)
-		+SVerticalBox::Slot()
-		.AutoHeight()
+		SNew(SHorizontalBox)
+		+SHorizontalBox::Slot()
+		.AutoWidth()
+		.VAlign(VAlign_Center)
 		[
-			MenuBarBuilder.MakeWidget()
+			SNew(STextBlock)
+			.Visibility(bHasTitle ? EVisibility::Visible : EVisibility::Collapsed)
+			.ColorAndOpacity(FEditorStyle::GetColor("MultiboxHookColor"))
+			.Text(InitOptions.ViewerTitleString)
 		]
+	]
+	+SVerticalBox::Slot()
+	.AutoHeight()
+	[
+		SAssignNew(SearchBox, SSearchBox)
+			.OnTextChanged( this, &SClassViewer::OnFilterTextChanged )
+			.OnTextCommitted( this, &SClassViewer::OnFilterTextCommitted )
+	]
+
+	+SVerticalBox::Slot()
+	.AutoHeight()
+	[
+		SNew(SSeparator)
+			.Visibility(HeaderVisibility)
+	]
 
-		+SVerticalBox::Slot()
-		.AutoHeight()
-		.Padding( 1.0f, 0.0f, 1.0f, 0.0f )
+	+SVerticalBox::Slot()
+	.FillHeight(1.0f)
+	[
+		SNew(SOverlay)
+
+		+SOverlay::Slot()
+		.HAlign(HAlign_Fill)
+		.VAlign(VAlign_Fill)
 		[
-			SNew(SHorizontalBox)
-			+SHorizontalBox::Slot()
-			.AutoWidth()
-			.VAlign(VAlign_Center)
+			SNew(SVerticalBox)
+
+			+SVerticalBox::Slot()
+			.FillHeight(1.0f)
 			[
-				SNew(STextBlock)
-				.Visibility(bHasTitle ? EVisibility::Visible : EVisibility::Collapsed)
-				.ColorAndOpacity(FEditorStyle::GetColor("MultiboxHookColor"))
-				.Text(InitOptions.ViewerTitleString)
-			]
-		]
-		+SVerticalBox::Slot()
-		.AutoHeight()
-		[
-			SAssignNew(SearchBox, SSearchBox)
-				.OnTextChanged( this, &SClassViewer::OnFilterTextChanged )
-				.OnTextCommitted( this, &SClassViewer::OnFilterTextCommitted )
-		]
+				SAssignNew(ClassTree, STreeView<TSharedPtr< FClassViewerNode > >)
+				.Visibility(InitOptions.DisplayMode == EClassViewerDisplayMode::TreeView ? EVisibility::Visible : EVisibility::Collapsed)
+	
+				.SelectionMode(ESelectionMode::Single)
 
-		+SVerticalBox::Slot()
-		.AutoHeight()
-		[
-			SNew(SSeparator)
-				.Visibility(HeaderVisibility)
-		]
+				.TreeItemsSource( &RootTreeItems )
+				// Called to child items for any given parent item
+				.OnGetChildren( this, &SClassViewer::OnGetChildrenForClassViewerTree )
 
-		+SVerticalBox::Slot()
-		.FillHeight(1.0f)
-		[
-			SNew(SOverlay)
+				// Called to handle recursively expanding/collapsing items
+				.OnSetExpansionRecursive(this, &SClassViewer::SetAllExpansionStates_Helper )
+
+				// Generates the actual widget for a tree item
+				.OnGenerateRow( this, &SClassViewer::OnGenerateRowForClassViewer ) 
+
+				// Generates the right click menu.
+				.OnContextMenuOpening( OnContextMenuOpening )
+
+				// Find out when the user selects something in the tree
+				.OnSelectionChanged( this, &SClassViewer::OnClassViewerSelectionChanged )
 
-			+SOverlay::Slot()
-			.HAlign(HAlign_Fill)
-			.VAlign(VAlign_Fill)
+				// Called when the expansion state of an item changes
+				.OnExpansionChanged( this, &SClassViewer::OnClassViewerExpansionChanged )
+
+				// Allow for some spacing between items with a larger item height.
+				.ItemHeight(20.0f)
+
+				.HeaderRow
+				(
+					SNew(SHeaderRow)
+					.Visibility(EVisibility::Collapsed)
+					+ SHeaderRow::Column(TEXT("Class"))
+					.DefaultLabel(NSLOCTEXT("ClassViewer", "Class", "Class"))
+				)
+			]
+
+			+SVerticalBox::Slot()
+			.FillHeight(1.0f)
 			[
-				SNew(SVerticalBox)
+				SAssignNew(ClassList, SListView<TSharedPtr< FClassViewerNode > >)
+					.Visibility(InitOptions.DisplayMode == EClassViewerDisplayMode::ListView ? EVisibility::Visible : EVisibility::Collapsed)
 
-				+SVerticalBox::Slot()
-				.FillHeight(1.0f)
-				[
-					SAssignNew(ClassTree, STreeView<TSharedPtr< FClassViewerNode > >)
-					.Visibility(InitOptions.DisplayMode == EClassViewerDisplayMode::TreeView ? EVisibility::Visible : EVisibility::Collapsed)
-	
 					.SelectionMode(ESelectionMode::Single)
 
-					.TreeItemsSource( &RootTreeItems )
-					// Called to child items for any given parent item
-					.OnGetChildren( this, &SClassViewer::OnGetChildrenForClassViewerTree )
-
-					// Called to handle recursively expanding/collapsing items
-					.OnSetExpansionRecursive(this, &SClassViewer::SetAllExpansionStates_Helper )
+					.ListItemsSource( &RootTreeItems )
 
 					// Generates the actual widget for a tree item
 					.OnGenerateRow( this, &SClassViewer::OnGenerateRowForClassViewer ) 
@@ -2172,9 +2208,6 @@ void SClassViewer::Construct(const FArguments& InArgs, const FClassViewerInitial
 					// Find out when the user selects something in the tree
 					.OnSelectionChanged( this, &SClassViewer::OnClassViewerSelectionChanged )
 
-					// Called when the expansion state of an item changes
-					.OnExpansionChanged( this, &SClassViewer::OnClassViewerExpansionChanged )
-
 					// Allow for some spacing between items with a larger item height.
 					.ItemHeight(20.0f)
 
@@ -2185,52 +2218,40 @@ void SClassViewer::Construct(const FArguments& InArgs, const FClassViewerInitial
 						+ SHeaderRow::Column(TEXT("Class"))
 						.DefaultLabel(NSLOCTEXT("ClassViewer", "Class", "Class"))
 					)
-				]
-
-				+SVerticalBox::Slot()
-				.FillHeight(1.0f)
-				[
-					SAssignNew(ClassList, SListView<TSharedPtr< FClassViewerNode > >)
-						.Visibility(InitOptions.DisplayMode == EClassViewerDisplayMode::ListView ? EVisibility::Visible : EVisibility::Collapsed)
-
-						.SelectionMode(ESelectionMode::Single)
-
-						.ListItemsSource( &RootTreeItems )
-
-						// Generates the actual widget for a tree item
-						.OnGenerateRow( this, &SClassViewer::OnGenerateRowForClassViewer ) 
-
-						// Generates the right click menu.
-						.OnContextMenuOpening( OnContextMenuOpening )
-
-						// Find out when the user selects something in the tree
-						.OnSelectionChanged( this, &SClassViewer::OnClassViewerSelectionChanged )
 
-						// Allow for some spacing between items with a larger item height.
-						.ItemHeight(20.0f)
-
-						.HeaderRow
-						(
-							SNew(SHeaderRow)
-							.Visibility(EVisibility::Collapsed)
-							+ SHeaderRow::Column(TEXT("Class"))
-							.DefaultLabel(NSLOCTEXT("ClassViewer", "Class", "Class"))
-						)
-
-				]
 			]
+		]
 
-			+SOverlay::Slot()
-			.HAlign(HAlign_Fill)
-			.VAlign(VAlign_Bottom)
-			.Padding(FMargin(24, 0, 24, 0))
-			[
-				// Asset discovery indicator
-				AssetDiscoveryIndicator
-			]
+		+SOverlay::Slot()
+		.HAlign(HAlign_Fill)
+		.VAlign(VAlign_Bottom)
+		.Padding(FMargin(24, 0, 24, 0))
+		[
+			// Asset discovery indicator
+			AssetDiscoveryIndicator
 		]
 	];
 
+	// When using a class picker in list-view mode, the widget will auto-focus the search box
+	// and allow the up and down arrow keys to navigate and enter to pick without using the mouse ever
+	if ( InitOptions.Mode == EClassViewerMode::ClassPicker && InitOptions.DisplayMode == EClassViewerDisplayMode::ListView )
+	{
+		this->ChildSlot
+			[
+				SNew(SListViewSelectorDropdownMenu<TSharedPtr<FClassViewerNode>>, SearchBox, ClassList)
+				[
+					ClassViewerContent.ToSharedRef()
+				]
+			];
+	}
+	else
+	{
+		this->ChildSlot
+			[
+				ClassViewerContent.ToSharedRef()
+			];
+	}
+
 	// Construct the class hierarchy.
 	ClassViewer::Helpers::ConstructClassHierarchy();
 
@@ -2271,6 +2292,12 @@ void SClassViewer::OnGetChildrenForClassViewerTree( TSharedPtr<FClassViewerNode>
 
 void SClassViewer::OnClassViewerSelectionChanged( TSharedPtr<FClassViewerNode> Item, ESelectInfo::Type SelectInfo )
 {
+	// Do not act on selection change when it is for navigation
+	if(SelectInfo == ESelectInfo::OnNavigation)
+	{
+		return;
+	}
+
 	// Sometimes the item is not valid anymore due to filtering.
 	if(Item.IsValid() == false || Item->IsRestricted())
 	{
@@ -2286,7 +2313,7 @@ void SClassViewer::OnClassViewerSelectionChanged( TSharedPtr<FClassViewerNode> I
 	{
 		UClass* Class = Item->Class.Get();
 
-		// If the class is not NULL and UnloadedBlueprintData is valid then attempt to load it. UnloadedBlueprintData is invalid in the case of a "None" item.
+		// If the class is NULL and UnloadedBlueprintData is valid then attempt to load it. UnloadedBlueprintData is invalid in the case of a "None" item.
 		if ( bEnableClassDynamicLoading && !Class && Item->UnloadedBlueprintData.IsValid() )
 		{
 			ClassViewer::Helpers::LoadClass( Item );
@@ -2596,10 +2623,11 @@ void SClassViewer::OnFilterTextCommitted(const FText& InText, ETextCommit::Type
 				FirstSelected = SelectedList[0];
 				Class = FirstSelected->Class.Get();
 
-				// If the class is not NULL and UnloadedBlueprintData is valid then attempt to load it. UnloadedBlueprintData is invalid in the case of a "None" item.
-				if ( bEnableClassDynamicLoading && Class && FirstSelected->UnloadedBlueprintData.IsValid())
+				// If the class is NULL and UnloadedBlueprintData is valid then attempt to load it. UnloadedBlueprintData is invalid in the case of a "None" item.
+				if ( bEnableClassDynamicLoading && Class == nullptr && FirstSelected->UnloadedBlueprintData.IsValid())
 				{
 					ClassViewer::Helpers::LoadClass(FirstSelected);
+					Class = FirstSelected->Class.Get();
 				}
 
 				// Check if the item passes the filter, parent items might be displayed but filtered out and thus not desired to be selected.
diff --git a/Engine/Source/Editor/GraphEditor/Private/KismetPins/SGraphPinObject.cpp b/Engine/Source/Editor/GraphEditor/Private/KismetPins/SGraphPinObject.cpp
--- a/Engine/Source/Editor/GraphEditor/Private/KismetPins/SGraphPinObject.cpp
+++ b/Engine/Source/Editor/GraphEditor/Private/KismetPins/SGraphPinObject.cpp
@@ -72,10 +72,7 @@ TSharedRef<SWidget>	SGraphPinObject::GetDefaultValueWidget()
 					.Text( this, &SGraphPinObject::OnGetComboTextValue )
 					.ToolTipText( this, &SGraphPinObject::GetObjectToolTip )
 				]
-				.MenuContent()
-				[
-					GenerateAssetPicker()
-				]
+				.OnGetMenuContent(this, &SGraphPinObject::GenerateAssetPicker)
 			]
 			// Use button
 			+SHorizontalBox::Slot()
