Made "Create New Blueprint" in Class Viewer more robust, now invalid names are disallowed.

diff --git a/Engine/Source/Editor/ClassViewer/Private/SClassViewer.cpp b/Engine/Source/Editor/ClassViewer/Private/SClassViewer.cpp
--- a/Engine/Source/Editor/ClassViewer/Private/SClassViewer.cpp
+++ b/Engine/Source/Editor/ClassViewer/Private/SClassViewer.cpp
@@ -445,6 +445,9 @@ namespace ClassViewer
 		/** true if the Class Hierarchy should be populated. */
 		static bool bPopulateClassHierarchy;
 
+		/** The currently selected path */
+		static FString NewBlueprintPath;
+
 		/** The cached filename we will use to create a new Blueprint */
 		static FString NewBlueprintFilename;
 
@@ -932,7 +935,7 @@ namespace ClassViewer
 			static FText MakeFullPathLabel()
 			{
 				FFormatNamedArguments Arguments;
-				Arguments.Add(TEXT("FullPath"), FText::FromString( NewBlueprintFilename ) );
+				Arguments.Add(TEXT("FullPath"), FText::FromString( NewBlueprintPath / NewBlueprintFilename ) );
 				return FText::Format( LOCTEXT( "PickNewBlueprintCreateLabel", "Create {FullPath}" ), Arguments );
 			}
 
@@ -943,13 +946,13 @@ namespace ClassViewer
 			*/
 			static void CreateBlueprintClicked(UClass* InCreationClass)
 			{
-				if(FPaths::GetBaseFilename(NewBlueprintFilename).IsEmpty())
+				if (!CanCreateBlueprint())
 				{
 					FMessageDialog::Open( EAppMsgType::Ok, LOCTEXT("PickNewBlueprintInvalidName", "Invalid name for Blueprint."));
 				}
 				else
 				{
-					CreateBlueprint(FText::FromString(NewBlueprintFilename), InCreationClass);
+					CreateBlueprint(FText::FromString(NewBlueprintPath / NewBlueprintFilename), InCreationClass);
 				}
 			}
 
@@ -960,7 +963,7 @@ namespace ClassViewer
 			*/
 			static void FilenameChanged(const FText& InFileName)
 			{
-				NewBlueprintFilename = FPaths::GetPath(NewBlueprintFilename) / InFileName.ToString();
+				NewBlueprintFilename = FText::TrimPrecedingAndTrailing(InFileName).ToString();
 			}
 
 			/**
@@ -970,7 +973,7 @@ namespace ClassViewer
 			*/
 			static void PathSelected(const FString& InPathName)
 			{
-				NewBlueprintFilename = InPathName / FPaths::GetBaseFilename(NewBlueprintFilename);
+				NewBlueprintPath = InPathName;
 			}
 
 			/**
@@ -978,7 +981,7 @@ namespace ClassViewer
 			*/
 			static bool CanCreateBlueprint()
 			{
-				return !FPaths::GetBaseFilename(NewBlueprintFilename).IsEmpty();
+				return !NewBlueprintFilename.IsEmpty() && FName(*NewBlueprintFilename).IsValidXName(INVALID_OBJECTNAME_CHARACTERS INVALID_LONGPACKAGE_CHARACTERS);
 			}
 
 			/**
@@ -992,13 +995,14 @@ namespace ClassViewer
 			{
 				if (CommitInfo == ETextCommit::OnEnter)
 				{
-					if(InBlueprintName.IsEmpty())
+					NewBlueprintFilename = FText::TrimPrecedingAndTrailing(InBlueprintName).ToString();
+					if (!CanCreateBlueprint())
 					{
 						FMessageDialog::Open( EAppMsgType::Ok, LOCTEXT("PickNewBlueprintInvalidName", "Invalid name for Blueprint."));
 					}
 					else
 					{
-						CreateBlueprint(FText::FromString(NewBlueprintFilename), InCreationClass);
+						CreateBlueprint(FText::FromString(NewBlueprintPath / NewBlueprintFilename), InCreationClass);
 					}
 				}
 			}
@@ -1011,7 +1015,7 @@ namespace ClassViewer
 			static TSharedRef<SWidget> MakeBlueprintNameWidget(UClass* InCreationClass)
 			{
 				return SNew(SEditableTextBox)
-					.Text(FText::FromString(FPaths::GetBaseFilename(NewBlueprintFilename)))
+					.Text(FText::FromString(NewBlueprintFilename))
 					.SelectAllTextWhenFocused(true)
 					.OnTextCommitted(FOnTextCommitted::CreateStatic(&BlueprintNameEntry::CreateBlueprintCommited, InCreationClass))
 					.OnTextChanged(FOnTextChanged::CreateStatic(&BlueprintNameEntry::FilenameChanged));
@@ -1027,7 +1031,7 @@ namespace ClassViewer
 				FContentBrowserModule& ContentBrowserModule = FModuleManager::Get().LoadModuleChecked<FContentBrowserModule>(TEXT("ContentBrowser"));
 
 				FPathPickerConfig PathPickerConfig;
-				PathPickerConfig.DefaultPath = NewBlueprintFilename;
+				PathPickerConfig.DefaultPath = NewBlueprintPath;
 				PathPickerConfig.bFocusSearchBoxWhenOpened = false;
 				PathPickerConfig.OnPathSelected = FOnPathSelected::CreateStatic(&BlueprintNameEntry::PathSelected);
 
@@ -1065,16 +1069,17 @@ namespace ClassViewer
 			else
 			{
 				// Reset cached filename
-				NewBlueprintFilename = TEXT("/Game");
+				NewBlueprintPath = TEXT("/Game");
 				if(IFileManager::Get().DirectoryExists(*(FPaths::GameContentDir() / TEXT("Blueprints"))))
 				{
-					NewBlueprintFilename /= TEXT("Blueprints");
+					NewBlueprintPath /= TEXT("Blueprints");
 				}
 				else
 				{
-					NewBlueprintFilename /= TEXT("Unsorted");
+					NewBlueprintPath /= TEXT("Unsorted");
 				}
-				NewBlueprintFilename /= TEXT("MyBlueprint");
+
+				NewBlueprintFilename = TEXT("MyBlueprint");
 
 				MenuBuilder.AddMenuEntry(
 					TAttribute<FText>::Create(TAttribute<FText>::FGetter::CreateStatic(&BlueprintNameEntry::MakeFullPathLabel)), 
