Pass on editor to correct cases of referencing singular "Blueprint" to be "Class Blueprint" when refering to Class Blueprints.

Changed: Content Browser, Class Viewer, and workflow when generating a Class Blueprint from a selection of actors.

diff --git a/Engine/Source/Developer/AssetTools/Private/AssetTypeActions/AssetTypeActions_Blueprint.h b/Engine/Source/Developer/AssetTools/Private/AssetTypeActions/AssetTypeActions_Blueprint.h
--- a/Engine/Source/Developer/AssetTools/Private/AssetTypeActions/AssetTypeActions_Blueprint.h
+++ b/Engine/Source/Developer/AssetTools/Private/AssetTypeActions/AssetTypeActions_Blueprint.h
@@ -8,7 +8,7 @@ class ASSETTOOLS_API FAssetTypeActions_Blueprint : public FAssetTypeActions_Base
 {
 public:
 	// IAssetTypeActions Implementation
-	virtual FText GetName() const override { return NSLOCTEXT("AssetTypeActions", "AssetTypeActions_Blueprint", "Blueprint"); }
+	virtual FText GetName() const override { return NSLOCTEXT("AssetTypeActions", "AssetTypeActions_Blueprint", "Class Blueprint"); }
 	virtual FColor GetTypeColor() const override { return FColor( 63, 126, 255 ); }
 	virtual UClass* GetSupportedClass() const override { return UBlueprint::StaticClass(); }
 	virtual bool HasActions ( const TArray<UObject*>& InObjects ) const override { return true; }
diff --git a/Engine/Source/Editor/ClassViewer/Private/SClassViewer.cpp b/Engine/Source/Editor/ClassViewer/Private/SClassViewer.cpp
--- a/Engine/Source/Editor/ClassViewer/Private/SClassViewer.cpp
+++ b/Engine/Source/Editor/ClassViewer/Private/SClassViewer.cpp
@@ -1083,7 +1083,7 @@ namespace ClassViewer
 
 				MenuBuilder.AddMenuEntry(
 					TAttribute<FText>::Create(TAttribute<FText>::FGetter::CreateStatic(&BlueprintNameEntry::MakeFullPathLabel)), 
-						LOCTEXT("PickNewBlueprintTooltip_Create", "Create new Blueprint"),
+						LOCTEXT("PickNewBlueprintTooltip_Create", "Create new Class Blueprint"),
 						FSlateIcon(), 
 						FUIAction(
 							FExecuteAction::CreateStatic( &BlueprintNameEntry::CreateBlueprintClicked, InCreationClass ),
@@ -1091,11 +1091,11 @@ namespace ClassViewer
 						)
 					);
 
-				MenuBuilder.BeginSection( NAME_None, LOCTEXT("PickNewBlueprintFilename", "Blueprint Filename") );
+				MenuBuilder.BeginSection( NAME_None, LOCTEXT("PickNewBlueprintFilename", "Class Blueprint Filename") );
 				MenuBuilder.AddWidget( BlueprintNameEntry::MakeBlueprintNameWidget(InCreationClass), FText::GetEmpty() );
 				MenuBuilder.EndSection();
 
-				MenuBuilder.BeginSection( NAME_None, LOCTEXT("PickNewBlueprintPath", "Blueprint Path") );
+				MenuBuilder.BeginSection( NAME_None, LOCTEXT("PickNewBlueprintPath", "Class Blueprint Path") );
 				MenuBuilder.AddWidget( BlueprintNameEntry::MakeBlueprintPathWidget(), FText::GetEmpty() );
 				MenuBuilder.EndSection();
 			}
@@ -1106,11 +1106,11 @@ namespace ClassViewer
 		{
 			if(InCreationClass->HasAnyClassFlags(CLASS_Deprecated))
 			{
-				return LOCTEXT("ClassViewerMenuCreateDeprecatedBlueprint_Tooltip", "Blueprint class is deprecated!");
+				return LOCTEXT("ClassViewerMenuCreateDeprecatedBlueprint_Tooltip", "Class is deprecated!");
 			}
 			else
 			{
-				return LOCTEXT("ClassViewerMenuCreateBlueprint_Tooltip", "Creates a Blueprint using this class as a base.");
+				return LOCTEXT("ClassViewerMenuCreateBlueprint_Tooltip", "Creates a Class Blueprint using this class as a base.");
 			}
 		}
 
@@ -1495,7 +1495,7 @@ private:
 						TAttribute<FText> DynamicTooltipAttribute = TAttribute<FText>::Create(DynamicTooltipGetter);
 
 						MenuBuilder.AddSubMenu(
-							LOCTEXT("ClassViewerMenuCreateBlueprint", "Create Blueprint"), 
+							LOCTEXT("ClassViewerMenuCreateBlueprint", "Create Class Blueprint"), 
 							DynamicTooltipAttribute, 
 							FNewMenuDelegate::CreateStatic( &ClassViewer::Helpers::OpenCreateBlueprintMenu, Class ),
 							FUIAction(
@@ -1514,7 +1514,7 @@ private:
 					MenuBuilder.BeginSection("ClassViewerDropDownHasBlueprint");
 					{
 						FUIAction Action( FExecuteAction::CreateStatic( &ClassViewer::Helpers::OpenBlueprintTool, ClassViewer::Helpers::GetBlueprint(Class) ) );
-						MenuBuilder.AddMenuEntry(LOCTEXT("ClassViewerMenuEditBlueprint", "Edit Blueprint..."), LOCTEXT("ClassViewerMenuEditBlueprint_Tooltip", "Open the Blueprint in the editor."), FSlateIcon(), Action);
+						MenuBuilder.AddMenuEntry(LOCTEXT("ClassViewerMenuEditBlueprint", "Edit Class Blueprint..."), LOCTEXT("ClassViewerMenuEditBlueprint_Tooltip", "Open the Class Blueprint in the editor."), FSlateIcon(), Action);
 					}
 					MenuBuilder.EndSection();
 
@@ -2409,7 +2409,7 @@ TSharedPtr< SWidget > SClassViewer::BuildMenuWidget()
 			TAttribute<FText> DynamicTooltipAttribute = TAttribute<FText>::Create(DynamicTooltipGetter);
 
 			MenuBuilder.AddSubMenu(
-				LOCTEXT("ClassViewerMenuCreateBlueprint", "Create Blueprint"), 
+				LOCTEXT("ClassViewerMenuCreateBlueprint", "Create Class Blueprint"), 
 				DynamicTooltipAttribute, 
 				FNewMenuDelegate::CreateStatic( &ClassViewer::Helpers::OpenCreateBlueprintMenu, RightClickClass ),
 				FUIAction(
@@ -2426,7 +2426,7 @@ TSharedPtr< SWidget > SClassViewer::BuildMenuWidget()
 			MenuBuilder.BeginSection("ClassViewerHasBlueprint");
 			{
 				FUIAction Action( FExecuteAction::CreateRaw( this, &SClassViewer::OnOpenBlueprintTool ) );
-				MenuBuilder.AddMenuEntry(LOCTEXT("ClassViewerMenuEditBlueprint", "Edit Blueprint..."), LOCTEXT("ClassViewerMenuEditBlueprint_Tooltip", "Open the Blueprint in the editor."), FSlateIcon(), Action);
+				MenuBuilder.AddMenuEntry(LOCTEXT("ClassViewerMenuEditBlueprint", "Edit Class Blueprint..."), LOCTEXT("ClassViewerMenuEditBlueprint_Tooltip", "Open the Class Blueprint in the editor."), FSlateIcon(), Action);
 			}
 			MenuBuilder.EndSection();
 
@@ -2712,7 +2712,7 @@ void SClassViewer::FillFilterEntries( FMenuBuilder& MenuBuilder )
 
 	MenuBuilder.BeginSection("ClassViewerFilterEntries2");
 	{
-		MenuBuilder.AddMenuEntry( LOCTEXT("BlueprintsOnly", "Blueprint Bases Only"), LOCTEXT( "BlueprinsOnly_Tooltip", "Filter the Class Viewer to show only base blueprint classes." ), FSlateIcon(), FUIAction(FExecuteAction::CreateRaw(this, &SClassViewer::MenuBlueprintBasesOnly_Execute), FCanExecuteAction::CreateRaw(this, &SClassViewer::Menu_CanExecute), FIsActionChecked::CreateRaw(this, &SClassViewer::MenuBlueprintBasesOnly_IsChecked)), NAME_None, EUserInterfaceActionType::Check );
+		MenuBuilder.AddMenuEntry( LOCTEXT("BlueprintsOnly", "Class Blueprint Bases Only"), LOCTEXT( "BlueprinsOnly_Tooltip", "Filter the Class Viewer to show only base blueprint classes." ), FSlateIcon(), FUIAction(FExecuteAction::CreateRaw(this, &SClassViewer::MenuBlueprintBasesOnly_Execute), FCanExecuteAction::CreateRaw(this, &SClassViewer::Menu_CanExecute), FIsActionChecked::CreateRaw(this, &SClassViewer::MenuBlueprintBasesOnly_IsChecked)), NAME_None, EUserInterfaceActionType::Check );
 	}
 	MenuBuilder.EndSection();
 }
diff --git a/Engine/Source/Editor/DetailCustomizations/Private/ActorDetails.cpp b/Engine/Source/Editor/DetailCustomizations/Private/ActorDetails.cpp
--- a/Engine/Source/Editor/DetailCustomizations/Private/ActorDetails.cpp
+++ b/Engine/Source/Editor/DetailCustomizations/Private/ActorDetails.cpp
@@ -1033,7 +1033,7 @@ void FActorDetails::AddSingleBlueprintRow( IDetailCategoryBuilder& BlueprintCate
 						.AutoWidth()
 						[
 							SNew(STextBlock)
-							.Text( FText::Format( LOCTEXT("CreateBlueprintFromActor", "Create {0} Blueprint..."), FText::FromString( SelectedActorClassName ) ) )
+							.Text( FText::Format( LOCTEXT("CreateBlueprintFromActor", "Create {0} Class Blueprint..."), FText::FromString( SelectedActorClassName ) ) )
 							.Font(IDetailLayoutBuilder::GetDetailFont())
 						]
 
diff --git a/Engine/Source/Editor/UMGEditor/Private/WidgetBlueprintEditorUtils.cpp b/Engine/Source/Editor/UMGEditor/Private/WidgetBlueprintEditorUtils.cpp
--- a/Engine/Source/Editor/UMGEditor/Private/WidgetBlueprintEditorUtils.cpp
+++ b/Engine/Source/Editor/UMGEditor/Private/WidgetBlueprintEditorUtils.cpp
@@ -163,8 +163,8 @@ void FWidgetBlueprintEditorUtils::CreateWidgetContextMenu(FMenuBuilder& MenuBuil
 	MenuBuilder.BeginSection("Actions");
 	{
 		MenuBuilder.AddMenuEntry(
-			LOCTEXT( "EditBlueprint_Label", "Edit Blueprint..." ),
-			LOCTEXT( "EditBlueprint_Tooltip", "Open the selected widget blueprint(s) for edit." ),
+			LOCTEXT( "EditBlueprint_Label", "Edit Widget Blueprint..." ),
+			LOCTEXT( "EditBlueprint_Tooltip", "Open the selected Widget Blueprint(s) for edit." ),
 			FSlateIcon(),
 			FUIAction(
 				FExecuteAction::CreateStatic( &FWidgetBlueprintEditorUtils::ExecuteOpenSelectedWidgetsForEdit, Widgets ),
diff --git a/Engine/Source/Runtime/Engine/Classes/Engine/Blueprint.h b/Engine/Source/Runtime/Engine/Classes/Engine/Blueprint.h
--- a/Engine/Source/Runtime/Engine/Classes/Engine/Blueprint.h
+++ b/Engine/Source/Runtime/Engine/Classes/Engine/Blueprint.h
@@ -30,9 +30,9 @@ UENUM()
 enum EBlueprintType
 {
 	// Normal blueprint
-	BPTYPE_Normal				UMETA(DisplayName="Blueprint"),
+	BPTYPE_Normal				UMETA(DisplayName="Class Blueprint"),
 	// Blueprint that is const during execution (no state graph and methods cannot modify member variables)
-	BPTYPE_Const				UMETA(DisplayName="Const Blueprint"),
+	BPTYPE_Const				UMETA(DisplayName="Const Class Blueprint"),
 	// Blueprint that serves as a container for macros to be used in other blueprints
 	BPTYPE_MacroLibrary			UMETA(DisplayName="Blueprint Macro Library"),
 	// Blueprint that serves as an interface to be implemented by other blueprints
