Fix for new Blueprints not immediately appearing in the Class Viewer (or Class Picker instances).

diff --git a/Engine/Source/Editor/ClassViewer/Private/SClassViewer.cpp b/Engine/Source/Editor/ClassViewer/Private/SClassViewer.cpp
--- a/Engine/Source/Editor/ClassViewer/Private/SClassViewer.cpp
+++ b/Engine/Source/Editor/ClassViewer/Private/SClassViewer.cpp
@@ -1839,7 +1839,7 @@ void FClassHierarchy::AddAsset(const FAssetData& InAddedAssetData)
 		TArray<FName> AncestorClassNames;
 		AssetRegistryModule.Get().GetAncestorClassNames(InAddedAssetData.AssetClass, AncestorClassNames);
 
-		if( AncestorClassNames.Contains(UBlueprint::StaticClass()->GetFName()) )
+		if( AncestorClassNames.Contains(UBlueprintCore::StaticClass()->GetFName()) )
 		{
 			// Make sure that the node does not already exist. There is a bit of double adding going on at times and this prevents it.
 			if(!FindNodeByGeneratedClassPackageName(ObjectClassRoot, InAddedAssetData.PackageName.ToString()).IsValid())
