Fixes a potential crash during editor shutdown relating to the ClassViewer.

PR #755

diff --git a/Engine/Source/Editor/ClassViewer/Private/ClassViewerModule.cpp b/Engine/Source/Editor/ClassViewer/Private/ClassViewerModule.cpp
--- a/Engine/Source/Editor/ClassViewer/Private/ClassViewerModule.cpp
+++ b/Engine/Source/Editor/ClassViewer/Private/ClassViewerModule.cpp
@@ -44,6 +44,8 @@ void FClassViewerModule::ShutdownModule()
 	{
 		FGlobalTabmanager::Get()->UnregisterNomadTabSpawner( ClassViewerModule::ClassViewerApp );
 	}
+
+	SClassViewer::DestroyClassHierarchy();
 }
 
 
diff --git a/Engine/Source/Editor/ClassViewer/Private/SClassViewer.cpp b/Engine/Source/Editor/ClassViewer/Private/SClassViewer.cpp
--- a/Engine/Source/Editor/ClassViewer/Private/SClassViewer.cpp
+++ b/Engine/Source/Editor/ClassViewer/Private/SClassViewer.cpp
@@ -576,6 +576,12 @@ namespace ClassViewer
 			}
 		}
 
+		/** Cleans up the Class Hierarchy */
+		static void DestroyClassHierachy()
+		{
+			ClassHierarchy.Reset();
+		}
+
 		/** Will populate the class hierarchy tree if previously requested. */
 		static void PopulateClassHierarchy()
 		{
@@ -2928,11 +2934,16 @@ FReply SClassViewer::OnFocusReceived( const FGeometry& MyGeometry, const FFocusE
 	return FReply::Unhandled();
 }
 
- bool SClassViewer::SupportsKeyboardFocus() const 
+bool SClassViewer::SupportsKeyboardFocus() const 
 {
 	return true;
 }
 
+void SClassViewer::DestroyClassHierarchy()
+{
+	ClassViewer::Helpers::DestroyClassHierachy();
+}
+
 TSharedPtr<FClassViewerNode> SClassViewer::CreateNoneOption()
 {
 	TSharedPtr<FClassViewerNode> NoneItem = MakeShareable( new FClassViewerNode("None", "None") );
diff --git a/Engine/Source/Editor/ClassViewer/Private/SClassViewer.h b/Engine/Source/Editor/ClassViewer/Private/SClassViewer.h
--- a/Engine/Source/Editor/ClassViewer/Private/SClassViewer.h
+++ b/Engine/Source/Editor/ClassViewer/Private/SClassViewer.h
@@ -42,6 +42,8 @@ public:
 	virtual FReply OnFocusReceived( const FGeometry& MyGeometry, const FFocusEvent& InFocusEvent ) override;
 	virtual bool SupportsKeyboardFocus() const override;
 
+	/** Destroys the internal Class Hierarchy database */
+	static void DestroyClassHierarchy();
 private:
 	/** Retrieves the children for the input node.
 	 *	@param InParent				The parent node to retrieve the children from.
