[INTEGRATE] Change 2412499 by [REDACTED_PERSON][REDACTED_HANDLE] on 2015/01/20 13:27:21

	Relabeled "Class Blueprint" as "Blueprint Class" throughout the editor and documentation.

diff --git a/Engine/Source/Developer/AssetTools/Private/AssetTypeActions/AssetTypeActions_Blueprint.h b/Engine/Source/Developer/AssetTools/Private/AssetTypeActions/AssetTypeActions_Blueprint.h
--- a/Engine/Source/Developer/AssetTools/Private/AssetTypeActions/AssetTypeActions_Blueprint.h
+++ b/Engine/Source/Developer/AssetTools/Private/AssetTypeActions/AssetTypeActions_Blueprint.h
@@ -8,7 +8,7 @@ class ASSETTOOLS_API FAssetTypeActions_Blueprint : public FAssetTypeActions_Base
 {
 public:
 	// IAssetTypeActions Implementation
-	virtual FText GetName() const override { return NSLOCTEXT("AssetTypeActions", "AssetTypeActions_Blueprint", "Class Blueprint"); }
+	virtual FText GetName() const override { return NSLOCTEXT("AssetTypeActions", "AssetTypeActions_Blueprint", "Blueprint Class"); }
 	virtual FColor GetTypeColor() const override { return FColor( 63, 126, 255 ); }
 	virtual UClass* GetSupportedClass() const override { return UBlueprint::StaticClass(); }
 	virtual bool HasActions ( const TArray<UObject*>& InObjects ) const override { return true; }
diff --git a/Engine/Source/Editor/ClassViewer/Private/SClassViewer.cpp b/Engine/Source/Editor/ClassViewer/Private/SClassViewer.cpp
--- a/Engine/Source/Editor/ClassViewer/Private/SClassViewer.cpp
+++ b/Engine/Source/Editor/ClassViewer/Private/SClassViewer.cpp
@@ -1094,7 +1094,7 @@ namespace ClassViewer
 
 				MenuBuilder.AddMenuEntry(
 					TAttribute<FText>::Create(TAttribute<FText>::FGetter::CreateStatic(&BlueprintNameEntry::MakeFullPathLabel)), 
-						LOCTEXT("PickNewBlueprintTooltip_Create", "Create new Class Blueprint"),
+						LOCTEXT("PickNewBlueprintTooltip_Create", "Create new Blueprint Class"),
 						FSlateIcon(), 
 						FUIAction(
 							FExecuteAction::CreateStatic( &BlueprintNameEntry::CreateBlueprintClicked, InCreationClass ),
@@ -1102,11 +1102,11 @@ namespace ClassViewer
 						)
 					);
 
-				MenuBuilder.BeginSection( NAME_None, LOCTEXT("PickNewBlueprintFilename", "Class Blueprint Filename") );
+				MenuBuilder.BeginSection( NAME_None, LOCTEXT("PickNewBlueprintFilename", "Blueprint Class Filename") );
 				MenuBuilder.AddWidget( BlueprintNameEntry::MakeBlueprintNameWidget(InCreationClass), FText::GetEmpty() );
 				MenuBuilder.EndSection();
 
-				MenuBuilder.BeginSection( NAME_None, LOCTEXT("PickNewBlueprintPath", "Class Blueprint Path") );
+				MenuBuilder.BeginSection( NAME_None, LOCTEXT("PickNewBlueprintPath", "Blueprint Class Path") );
 				MenuBuilder.AddWidget( BlueprintNameEntry::MakeBlueprintPathWidget(), FText::GetEmpty() );
 				MenuBuilder.EndSection();
 			}
@@ -1121,7 +1121,7 @@ namespace ClassViewer
 			}
 			else
 			{
-				return LOCTEXT("ClassViewerMenuCreateBlueprint_Tooltip", "Creates a Class Blueprint using this class as a base.");
+				return LOCTEXT("ClassViewerMenuCreateBlueprint_Tooltip", "Creates a Blueprint Class using this class as a base.");
 			}
 		}
 
@@ -1488,7 +1488,7 @@ private:
 						TAttribute<FText> DynamicTooltipAttribute = TAttribute<FText>::Create(DynamicTooltipGetter);
 
 						MenuBuilder.AddSubMenu(
-							LOCTEXT("ClassViewerMenuCreateBlueprint", "Create Class Blueprint"), 
+							LOCTEXT("ClassViewerMenuCreateBlueprint", "Create Blueprint Class"), 
 							DynamicTooltipAttribute, 
 							FNewMenuDelegate::CreateStatic( &ClassViewer::Helpers::OpenCreateBlueprintMenu, Class ),
 							FUIAction(
@@ -1507,7 +1507,7 @@ private:
 					MenuBuilder.BeginSection("ClassViewerDropDownHasBlueprint");
 					{
 						FUIAction Action( FExecuteAction::CreateStatic( &ClassViewer::Helpers::OpenBlueprintTool, ClassViewer::Helpers::GetBlueprint(Class) ) );
-						MenuBuilder.AddMenuEntry(LOCTEXT("ClassViewerMenuEditBlueprint", "Edit Class Blueprint..."), LOCTEXT("ClassViewerMenuEditBlueprint_Tooltip", "Open the Class Blueprint in the editor."), FSlateIcon(), Action);
+						MenuBuilder.AddMenuEntry(LOCTEXT("ClassViewerMenuEditBlueprint", "Edit Blueprint Class..."), LOCTEXT("ClassViewerMenuEditBlueprint_Tooltip", "Open the Blueprint Class in the editor."), FSlateIcon(), Action);
 					}
 					MenuBuilder.EndSection();
 
@@ -2402,7 +2402,7 @@ TSharedPtr< SWidget > SClassViewer::BuildMenuWidget()
 			TAttribute<FText> DynamicTooltipAttribute = TAttribute<FText>::Create(DynamicTooltipGetter);
 
 			MenuBuilder.AddSubMenu(
-				LOCTEXT("ClassViewerMenuCreateBlueprint", "Create Class Blueprint"), 
+				LOCTEXT("ClassViewerMenuCreateBlueprint", "Create Blueprint Class"), 
 				DynamicTooltipAttribute, 
 				FNewMenuDelegate::CreateStatic( &ClassViewer::Helpers::OpenCreateBlueprintMenu, RightClickClass ),
 				FUIAction(
@@ -2419,7 +2419,7 @@ TSharedPtr< SWidget > SClassViewer::BuildMenuWidget()
 			MenuBuilder.BeginSection("ClassViewerHasBlueprint");
 			{
 				FUIAction Action( FExecuteAction::CreateRaw( this, &SClassViewer::OnOpenBlueprintTool ) );
-				MenuBuilder.AddMenuEntry(LOCTEXT("ClassViewerMenuEditBlueprint", "Edit Class Blueprint..."), LOCTEXT("ClassViewerMenuEditBlueprint_Tooltip", "Open the Class Blueprint in the editor."), FSlateIcon(), Action);
+				MenuBuilder.AddMenuEntry(LOCTEXT("ClassViewerMenuEditBlueprint", "Edit Blueprint Class..."), LOCTEXT("ClassViewerMenuEditBlueprint_Tooltip", "Open the Blueprint Class in the editor."), FSlateIcon(), Action);
 			}
 			MenuBuilder.EndSection();
 
@@ -2705,7 +2705,7 @@ void SClassViewer::FillFilterEntries( FMenuBuilder& MenuBuilder )
 
 	MenuBuilder.BeginSection("ClassViewerFilterEntries2");
 	{
-		MenuBuilder.AddMenuEntry( LOCTEXT("BlueprintsOnly", "Class Blueprint Bases Only"), LOCTEXT( "BlueprinsOnly_Tooltip", "Filter the Class Viewer to show only base blueprint classes." ), FSlateIcon(), FUIAction(FExecuteAction::CreateRaw(this, &SClassViewer::MenuBlueprintBasesOnly_Execute), FCanExecuteAction::CreateRaw(this, &SClassViewer::Menu_CanExecute), FIsActionChecked::CreateRaw(this, &SClassViewer::MenuBlueprintBasesOnly_IsChecked)), NAME_None, EUserInterfaceActionType::Check );
+		MenuBuilder.AddMenuEntry( LOCTEXT("BlueprintsOnly", "Blueprint Class Bases Only"), LOCTEXT( "BlueprinsOnly_Tooltip", "Filter the Class Viewer to show only base blueprint classes." ), FSlateIcon(), FUIAction(FExecuteAction::CreateRaw(this, &SClassViewer::MenuBlueprintBasesOnly_Execute), FCanExecuteAction::CreateRaw(this, &SClassViewer::Menu_CanExecute), FIsActionChecked::CreateRaw(this, &SClassViewer::MenuBlueprintBasesOnly_IsChecked)), NAME_None, EUserInterfaceActionType::Check );
 	}
 	MenuBuilder.EndSection();
 }
diff --git a/Engine/Source/Editor/DetailCustomizations/Private/ActorDetails.cpp b/Engine/Source/Editor/DetailCustomizations/Private/ActorDetails.cpp
--- a/Engine/Source/Editor/DetailCustomizations/Private/ActorDetails.cpp
+++ b/Engine/Source/Editor/DetailCustomizations/Private/ActorDetails.cpp
@@ -776,14 +776,14 @@ void FActorDetails::AddUtilityBlueprintRows( IDetailCategoryBuilder& BlueprintCa
 			.WidthOverride( 200 )
 			[
 				SNew( SButton )
-				.ToolTipText( LOCTEXT( "CreateHarvestBlueprint_ToolTip", "Copy the components from the selected actors and create a new Class Blueprint" ) )
+				.ToolTipText( LOCTEXT( "CreateHarvestBlueprint_ToolTip", "Copy the components from the selected actors and create a new Blueprint Class" ) )
 				.HAlign( HAlign_Fill )
 				//.ButtonColorAndOpacity( FLinearColor( 0.2f, 0.4f, 0.6f, 1.0f ) )
 				.OnClicked( FOnClicked::CreateRaw( this, &FActorDetails::OnPickBlueprintPathClicked, true ) )
 				.ContentPadding( 2 )
 				[
 					SNew( STextBlock )
-					.Text( LOCTEXT( "ReplaceWithAmalgamBlueprint", "Convert to Class Blueprint" ) )
+					.Text( LOCTEXT( "ReplaceWithAmalgamBlueprint", "Convert to Blueprint Class" ) )
 					.Font(IDetailLayoutBuilder::GetDetailFont())
 				]
 			]
@@ -1019,7 +1019,7 @@ void FActorDetails::AddSingleBlueprintRow( IDetailCategoryBuilder& BlueprintCate
 						.AutoWidth()
 						[
 							SNew(STextBlock)
-							.Text( FText::Format( LOCTEXT("CreateBlueprintFromActor", "Create {0} Class Blueprint..."), FText::FromString( SelectedActorClassName ) ) )
+							.Text( FText::Format( LOCTEXT("CreateBlueprintFromActor", "Create {0} Blueprint Class..."), FText::FromString( SelectedActorClassName ) ) )
 							.Font(IDetailLayoutBuilder::GetDetailFont())
 						]
 
diff --git a/Engine/Source/Editor/LevelEditor/Private/LevelEditorActions.cpp b/Engine/Source/Editor/LevelEditor/Private/LevelEditorActions.cpp
--- a/Engine/Source/Editor/LevelEditor/Private/LevelEditorActions.cpp
+++ b/Engine/Source/Editor/LevelEditor/Private/LevelEditorActions.cpp
@@ -1967,9 +1967,9 @@ void FLevelEditorActionCallbacks::OpenLevelBlueprint( TWeakPtr< SLevelEditor > L
 	}
 }
 
-void FLevelEditorActionCallbacks::CreateClassBlueprint()
+void FLevelEditorActionCallbacks::CreateBlueprintClass()
 {
-	// Use the BlueprintFactory to allow the user to pick a parent class for the new class Blueprint
+	// Use the BlueprintFactory to allow the user to pick a parent class for the new Blueprint class
 	UBlueprintFactory* NewFactory = Cast<UBlueprintFactory>(ConstructObject<UFactory>( UBlueprintFactory::StaticClass() ));
 	FEditorDelegates::OnConfigureNewAssetProperties.Broadcast(NewFactory);
 	if ( NewFactory->ConfigureProperties() )
@@ -1977,7 +1977,7 @@ void FLevelEditorActionCallbacks::CreateClassBlueprint()
 		UClass* SelectedClass = NewFactory->ParentClass;
 
 		// Now help the user pick a path and name for the new Blueprint
-		UBlueprint* Blueprint = FKismetEditorUtilities::CreateBlueprintFromClass(NSLOCTEXT("LevelEditorCommands", "CreateClassBlueprint_Title", "Create Class Blueprint"), SelectedClass);
+		UBlueprint* Blueprint = FKismetEditorUtilities::CreateBlueprintFromClass(NSLOCTEXT("LevelEditorCommands", "CreateBlueprintClass_Title", "Create Blueprint Class"), SelectedClass);
 
 		if( Blueprint )
 		{
@@ -2856,7 +2856,7 @@ void FLevelEditorCommands::RegisterCommands()
 
 	UI_COMMAND( OpenLevelBlueprint, "Open Level Blueprint", "Edit the Level Blueprint for the current level", EUserInterfaceActionType::Button, FInputGesture() );
 	UI_COMMAND( CheckOutProjectSettingsConfig, "Check Out", "Checks out the project settings config file so the game mode can be set.", EUserInterfaceActionType::Button, FInputGesture() );
-	UI_COMMAND( CreateClassBlueprint, "New Class Blueprint...", "Create a new Class Blueprint", EUserInterfaceActionType::Button, FInputGesture());
+	UI_COMMAND( CreateBlueprintClass, "New Blueprint Class...", "Create a new Blueprint Class", EUserInterfaceActionType::Button, FInputGesture());
 
 	UI_COMMAND( ShowTransformWidget, "Show Transform Widget", "Toggles the visibility of the transform widgets", EUserInterfaceActionType::ToggleButton, FInputGesture() );
 	UI_COMMAND( AllowTranslucentSelection, "Allow Translucent Selection", "Allows translucent objects to be selected", EUserInterfaceActionType::ToggleButton, FInputGesture(EKeys::T) );
diff --git a/Engine/Source/Editor/LevelEditor/Private/LevelEditorToolBar.cpp b/Engine/Source/Editor/LevelEditor/Private/LevelEditorToolBar.cpp
--- a/Engine/Source/Editor/LevelEditor/Private/LevelEditorToolBar.cpp
+++ b/Engine/Source/Editor/LevelEditor/Private/LevelEditorToolBar.cpp
@@ -1861,7 +1861,7 @@ TSharedRef< SWidget > FLevelEditorToolBar::GenerateOpenBlueprintMenuContent( TSh
 
 
 		/** Generates 'open blueprint' sub-menu */
-		static void MakeOpenClassBPMenu(FMenuBuilder& InMenuBuilder)
+		static void MakeOpenBPClassMenu(FMenuBuilder& InMenuBuilder)
 		{
 			FContentBrowserModule& ContentBrowserModule = FModuleManager::Get().LoadModuleChecked<FContentBrowserModule>(TEXT("ContentBrowser"));
 
@@ -1934,17 +1934,17 @@ TSharedRef< SWidget > FLevelEditorToolBar::GenerateOpenBlueprintMenuContent( TSh
 	}
 	MenuBuilder.EndSection();
 
-	MenuBuilder.BeginSection(NAME_None, LOCTEXT("ClassBlueprints", "Class Blueprints"));
+	MenuBuilder.BeginSection(NAME_None, LOCTEXT("Blueprints Class", "Blueprints Class"));
 	{
-		// New Class Blueprint...
-		MenuBuilder.AddMenuEntry(FLevelEditorCommands::Get().CreateClassBlueprint, NAME_None, LOCTEXT("NewClassBlueprint", "New Class Blueprint..."));
+		// New Blueprint Class...
+		MenuBuilder.AddMenuEntry(FLevelEditorCommands::Get().CreateBlueprintClass, NAME_None, LOCTEXT("NewBlueprintClass", "New Blueprint Class..."));
 
-		// Open Class Blueprint...
+		// Open Blueprint Class...
 		FSlateIcon OpenBPIcon(FEditorStyle::GetStyleSetName(), "LevelEditor.OpenClassBlueprint");
 		MenuBuilder.AddSubMenu(
-			LOCTEXT("OpenClassBlueprintSubMenu", "Open Class Blueprint..."),
-			LOCTEXT("OpenClassBlueprintSubMenu_ToolTip", "Open an existing Class Blueprint in this project"),
-			FNewMenuDelegate::CreateStatic(&FBlueprintMenus::MakeOpenClassBPMenu), 
+			LOCTEXT("OpenBlueprintClassSubMenu", "Open Blueprint Class..."),
+			LOCTEXT("OpenBlueprintClassSubMenu_ToolTip", "Open an existing Blueprint Class in this project"),
+			FNewMenuDelegate::CreateStatic(&FBlueprintMenus::MakeOpenBPClassMenu), 
 			false, 
 			OpenBPIcon );
 	}
diff --git a/Engine/Source/Editor/LevelEditor/Private/SLevelEditor.cpp b/Engine/Source/Editor/LevelEditor/Private/SLevelEditor.cpp
--- a/Engine/Source/Editor/LevelEditor/Private/SLevelEditor.cpp
+++ b/Engine/Source/Editor/LevelEditor/Private/SLevelEditor.cpp
@@ -97,8 +97,8 @@ void SLevelEditor::BindCommands()
 		FExecuteAction::CreateStatic< TWeakPtr< SLevelEditor > >( &FLevelEditorActionCallbacks::OpenLevelBlueprint, SharedThis( this ) ) );
 	
 	LevelEditorCommands->MapAction(
-		Actions.CreateClassBlueprint,
-		FExecuteAction::CreateStatic( &FLevelEditorActionCallbacks::CreateClassBlueprint ) );
+		Actions.CreateBlueprintClass,
+		FExecuteAction::CreateStatic( &FLevelEditorActionCallbacks::CreateBlueprintClass ) );
 
 	LevelEditorCommands->MapAction(
 		Actions.OpenContentBrowser,
diff --git a/Engine/Source/Editor/LevelEditor/Public/LevelEditorActions.h b/Engine/Source/Editor/LevelEditor/Public/LevelEditorActions.h
--- a/Engine/Source/Editor/LevelEditor/Public/LevelEditorActions.h
+++ b/Engine/Source/Editor/LevelEditor/Public/LevelEditorActions.h
@@ -503,7 +503,7 @@ public:
 	 */
 	TSharedPtr< FUICommandInfo > OpenLevelBlueprint;
 	TSharedPtr< FUICommandInfo > CheckOutProjectSettingsConfig;
-	TSharedPtr< FUICommandInfo > CreateClassBlueprint;
+	TSharedPtr< FUICommandInfo > CreateBlueprintClass;
 
 	/** Editor mode commands */
 	TArray< TSharedPtr< FUICommandInfo > > EditorModeCommands;
@@ -1070,8 +1070,8 @@ public:
 	/** Returns TRUE if the user can edit the game mode Blueprint, this requires the DefaultEngine config file to be writable */
 	static bool CanSelectGameModeBlueprint();
 
-	/** Helps the user create a class Blueprint */
-	static void CreateClassBlueprint();
+	/** Helps the user create a Blueprint class */
+	static void CreateBlueprintClass();
 
 	/** Shows only selected actors, hiding any unselected actors and unhiding any selected hidden actors. */
 	static void OnShowOnlySelectedActors();
diff --git a/Engine/Source/Editor/UnrealEd/Private/Commandlets/DumpBlueprintsInfoCommandlet.cpp b/Engine/Source/Editor/UnrealEd/Private/Commandlets/DumpBlueprintsInfoCommandlet.cpp
--- a/Engine/Source/Editor/UnrealEd/Private/Commandlets/DumpBlueprintsInfoCommandlet.cpp
+++ b/Engine/Source/Editor/UnrealEd/Private/Commandlets/DumpBlueprintsInfoCommandlet.cpp
@@ -1460,7 +1460,7 @@ static void DumpBlueprintInfoUtils::DumpInfoForClass(uint32 Indent, UClass* Blue
 		{
 			if (!FBlueprintEditorUtils::ImplementNewInterface(TempBlueprint, CommandOptions.InterfaceClass->GetFName()))
 			{
-				UE_LOG(LogBlueprintInfoDump, Warning, TEXT("Failed to add interface (%s), to class blueprint: '%s'"), *CommandOptions.InterfaceClass->GetName(), *ClassName);
+				UE_LOG(LogBlueprintInfoDump, Warning, TEXT("Failed to add interface (%s), to blueprint class: '%s'"), *CommandOptions.InterfaceClass->GetName(), *ClassName);
 			}
 		}		
 	}
diff --git a/Engine/Source/Runtime/Engine/Classes/Engine/Blueprint.h b/Engine/Source/Runtime/Engine/Classes/Engine/Blueprint.h
--- a/Engine/Source/Runtime/Engine/Classes/Engine/Blueprint.h
+++ b/Engine/Source/Runtime/Engine/Classes/Engine/Blueprint.h
@@ -37,9 +37,9 @@ UENUM()
 enum EBlueprintType
 {
 	// Normal blueprint
-	BPTYPE_Normal				UMETA(DisplayName="Class Blueprint"),
+	BPTYPE_Normal				UMETA(DisplayName="Blueprint Class"),
 	// Blueprint that is const during execution (no state graph and methods cannot modify member variables)
-	BPTYPE_Const				UMETA(DisplayName="Const Class Blueprint"),
+	BPTYPE_Const				UMETA(DisplayName="Const Blueprint Class"),
 	// Blueprint that serves as a container for macros to be used in other blueprints
 	BPTYPE_MacroLibrary			UMETA(DisplayName="Blueprint Macro Library"),
 	// Blueprint that serves as an interface to be implemented by other blueprints
