Reworked "Create Blueprint Class" dialog in class viewer, and removed redundancy from the code.

diff --git a/Engine/Source/Editor/ClassViewer/ClassViewer.Build.cs b/Engine/Source/Editor/ClassViewer/ClassViewer.Build.cs
--- a/Engine/Source/Editor/ClassViewer/ClassViewer.Build.cs
+++ b/Engine/Source/Editor/ClassViewer/ClassViewer.Build.cs
@@ -9,6 +9,7 @@ public class ClassViewer : ModuleRules
 		PrivateIncludePathModuleNames.AddRange(
 			new string[] {
 				"AssetRegistry",
+                "AssetTools",
 				"EditorWidgets",
 				"GameProjectGeneration",
                 "PropertyEditor",
@@ -32,6 +33,7 @@ public class ClassViewer : ModuleRules
 		DynamicallyLoadedModuleNames.AddRange(
 			new string[] {
 				"AssetRegistry",
+                "AssetTools",
 				"EditorWidgets",
 				"GameProjectGeneration",
 				"PropertyEditor",
diff --git a/Engine/Source/Editor/ClassViewer/Private/SClassViewer.cpp b/Engine/Source/Editor/ClassViewer/Private/SClassViewer.cpp
--- a/Engine/Source/Editor/ClassViewer/Private/SClassViewer.cpp
+++ b/Engine/Source/Editor/ClassViewer/Private/SClassViewer.cpp
@@ -16,6 +16,7 @@
 
 #include "AssetRegistryModule.h"
 #include "AssetSelection.h"
+#include "AssetToolsModule.h"
 
 #include "ClassViewerNode.h"
 
@@ -450,12 +451,6 @@ namespace ClassViewer
 		/** true if the Class Hierarchy should be populated. */
 		static bool bPopulateClassHierarchy;
 
-		/** The currently selected path */
-		static FString NewBlueprintPath;
-
-		/** The cached filename we will use to create a new Blueprint */
-		static FString NewBlueprintFilename;
-
 		// Pre-declare these functions.
 		static bool CheckIfBlueprintBase( TSharedPtr< FClassViewerNode> InNode );
 		static UBlueprint* GetBlueprint( UClass* InClass );
@@ -892,7 +887,7 @@ namespace ClassViewer
 		 * @param	InCreationClass			The class to create the blueprint from.
 		 * @param	InParentContent			The content to parent the STextEntryPopup to. 
 		 */
-		static void CreateBlueprint(const FText& InBlueprintName, UClass* InCreationClass)
+		static void CreateBlueprint(const FString& InBlueprintName, UClass* InCreationClass)
 		{
 			if(InCreationClass == NULL || !FKismetEditorUtilities::CanCreateBlueprintOfClass(InCreationClass))
 			{
@@ -901,7 +896,7 @@ namespace ClassViewer
 			}
 
 			// Get the full name of where we want to create the physics asset.
-			FString PackageName = InBlueprintName.ToString();
+			FString PackageName = InBlueprintName;
 
 			// Then find/create it.
 			UPackage* Package = CreatePackage(NULL, *PackageName);
@@ -938,177 +933,49 @@ namespace ClassViewer
 			RefreshAll();
 		}
 
-		struct BlueprintNameEntry
+		/**
+		 * Creates a SaveAssetDialog for specifying the path for the new blueprint
+		 */
+		static void OpenCreateBlueprintDialog(UClass* InCreationClass)
 		{
-			/**
-			* Build the full path label from the path & filename
-			*/
-			static FText MakeFullPathLabel()
-			{
-				FFormatNamedArguments Arguments;
-				Arguments.Add(TEXT("FullPath"), FText::FromString( NewBlueprintPath / NewBlueprintFilename ) );
-				return FText::Format( LOCTEXT( "PickNewBlueprintCreateLabel", "Create {FullPath}" ), Arguments );
-			}
-
-			/**
-			* Called when the create button (menu entry) is clicked
-			*
-			* @param	InBlueprintName			The package name and path for the new blueprint.
-			*/
-			static void CreateBlueprintClicked(UClass* InCreationClass)
-			{
-				if (!CanCreateBlueprint())
-				{
-					FMessageDialog::Open( EAppMsgType::Ok, LOCTEXT("PickNewBlueprintInvalidName", "Invalid name for Blueprint."));
-				}
-				else
-				{
-					CreateBlueprint(FText::FromString(NewBlueprintPath / NewBlueprintFilename), InCreationClass);
-				}
-			}
-
-			/**
-			* Called when the blueprint name entry is changed
-			*
-			* @param	InPathName				The file name for the new Blueprint
-			*/
-			static void FilenameChanged(const FText& InFileName)
-			{
-				NewBlueprintFilename = FText::TrimPrecedingAndTrailing(InFileName).ToString();
-			}
-
-			/**
-			* Called when the path selection changes in the PathPicker
-			*
-			* @param	InPathName				The path name for the new Blueprint
-			*/
-			static void PathSelected(const FString& InPathName)
-			{
-				NewBlueprintPath = InPathName;
-			}
+			// Determine default path for the Save Asset dialog
+			FString DefaultPath;
+			const FString DefaultDirectory = FEditorDirectories::Get().GetLastDirectory(ELastDirectory::NEW_ASSET);
+			FPackageName::TryConvertFilenameToLongPackageName(DefaultDirectory, DefaultPath);
 
-			/**
-			* Checks whether we can create a Blueprint with the information the user has input
-			*/
-			static bool CanCreateBlueprint()
+			if (DefaultPath.IsEmpty())
 			{
-				return !NewBlueprintFilename.IsEmpty() && FName(*NewBlueprintFilename).IsValidXName(INVALID_OBJECTNAME_CHARACTERS INVALID_LONGPACKAGE_CHARACTERS);
+				DefaultPath = TEXT("/Game/Blueprints");
 			}
 
-			/**
-			* Called when the blueprint name entry is committed.
-			*
-			* @param	InBlueprintName			The package name and path for the new blueprint.
-			* @param	CommitInfo				The method of committing.
-			* @param	InCreationClass			The class to create the blueprint from.
-			*/
-			static void CreateBlueprintCommited(const FText& InBlueprintName, ETextCommit::Type CommitInfo, UClass* InCreationClass)
-			{
-				if (CommitInfo == ETextCommit::OnEnter)
-				{
-					NewBlueprintFilename = FText::TrimPrecedingAndTrailing(InBlueprintName).ToString();
-					if (!CanCreateBlueprint())
-					{
-						FMessageDialog::Open( EAppMsgType::Ok, LOCTEXT("PickNewBlueprintInvalidName", "Invalid name for Blueprint."));
-					}
-					else
-					{
-						CreateBlueprint(FText::FromString(NewBlueprintPath / NewBlueprintFilename), InCreationClass);
-					}
-				}
-			}
+			// Determine default filename for the Save Asset dialog
+			check(InCreationClass != nullptr);
+			const FString ClassName = InCreationClass->ClassGeneratedBy ? InCreationClass->ClassGeneratedBy->GetName() : InCreationClass->GetName();
+			FString DefaultName = LOCTEXT("PrefixNew", "New").ToString() + ClassName;
 
-			/**
-			* Make the name entry widget
-			*
-			* @param	InCreationClass			The class to create the blueprint from.
-			*/
-			static TSharedRef<SWidget> MakeBlueprintNameWidget(UClass* InCreationClass)
-			{
-				return SNew(SEditableTextBox)
-					.Text(FText::FromString(NewBlueprintFilename))
-					.SelectAllTextWhenFocused(true)
-					.OnTextCommitted(FOnTextCommitted::CreateStatic(&BlueprintNameEntry::CreateBlueprintCommited, InCreationClass))
-					.OnTextChanged(FOnTextChanged::CreateStatic(&BlueprintNameEntry::FilenameChanged));
-			}
-
-			/**
-			* Make the path entry widget
-			*
-			* @param	InCreationClass			The class to create the blueprint from.
-			*/
-			static TSharedRef<SWidget> MakeBlueprintPathWidget()
-			{
-				FContentBrowserModule& ContentBrowserModule = FModuleManager::Get().LoadModuleChecked<FContentBrowserModule>(TEXT("ContentBrowser"));
+			FAssetToolsModule& AssetToolsModule = FModuleManager::LoadModuleChecked<FAssetToolsModule>("AssetTools");
+			FString UniquePackageName;
+			FString UniqueAssetName;
+			AssetToolsModule.Get().CreateUniqueAssetName(DefaultPath / DefaultName, TEXT(""), UniquePackageName, UniqueAssetName);
+			DefaultName = FPaths::GetCleanFilename(UniqueAssetName);
 
-				FPathPickerConfig PathPickerConfig;
-				PathPickerConfig.DefaultPath = NewBlueprintPath;
-				PathPickerConfig.bFocusSearchBoxWhenOpened = false;
-				PathPickerConfig.OnPathSelected = FOnPathSelected::CreateStatic(&BlueprintNameEntry::PathSelected);
-
-				return SNew(SBox)
-					.WidthOverride(300)
-					.HeightOverride(400)
-					[
-						ContentBrowserModule.Get().CreatePathPicker(PathPickerConfig)
-					];
-			}
-
-			/**
-			* Displays an error message in the menu if we cannot create a blueprint from the 
-			* passed-in class
-			*/
-			static TSharedRef<SWidget> MakeErrorWarningWidget()
-			{
-				return SNew(STextBlock)
-					.Text(LOCTEXT("PickNewBlueprintInvalidClass", "Invalid class to make a Blueprint of"));
-			}
-		};
+			// Initialize SaveAssetDialog config
+			FSaveAssetDialogConfig SaveAssetDialogConfig;
+			SaveAssetDialogConfig.DialogTitleOverride = LOCTEXT("CreateBlueprintDialogTitle", "Create Blueprint Class");
+			SaveAssetDialogConfig.DefaultPath = DefaultPath;
+			SaveAssetDialogConfig.DefaultAssetName = DefaultName;
+			SaveAssetDialogConfig.ExistingAssetPolicy = ESaveAssetDialogExistingAssetPolicy::AllowButWarn;
 
-		/**
-		 * Creates a STextEntryPopup for creating a blueprint from a class.
-		 *
-		 * @param	MenuBuilder				The MenuBuilder to use
-		 * @param	InParentContent			The content to parent the STextEntryPopup to. 
-		 */
-		static void OpenCreateBlueprintMenu(FMenuBuilder& MenuBuilder, UClass* InCreationClass)
-		{
-			if(InCreationClass == NULL || !CanCreateBlueprintOfClass_IgnoreDeprecation(InCreationClass))
+			FContentBrowserModule& ContentBrowserModule = FModuleManager::LoadModuleChecked<FContentBrowserModule>("ContentBrowser");
+			FString SaveObjectPath = ContentBrowserModule.Get().CreateModalSaveAssetDialog(SaveAssetDialogConfig);
+			if (!SaveObjectPath.IsEmpty())
 			{
-				MenuBuilder.AddWidget( BlueprintNameEntry::MakeErrorWarningWidget(), FText::GetEmpty() );
-			}
-			else
-			{
-				// Reset cached filename
-				NewBlueprintPath = TEXT("/Game");
-				if(IFileManager::Get().DirectoryExists(*(FPaths::GameContentDir() / TEXT("Blueprints"))))
-				{
-					NewBlueprintPath /= TEXT("Blueprints");
-				}
-				else
-				{
-					NewBlueprintPath /= TEXT("Unsorted");
-				}
+				const FString PackageName = FPackageName::ObjectPathToPackageName(SaveObjectPath);
+				const FString PackageFilename = FPackageName::LongPackageNameToFilename(PackageName);
+				const FString PackagePath = FPaths::GetPath(PackageFilename);
 
-				NewBlueprintFilename = TEXT("MyBlueprint");
-
-				MenuBuilder.AddMenuEntry(
-					TAttribute<FText>::Create(TAttribute<FText>::FGetter::CreateStatic(&BlueprintNameEntry::MakeFullPathLabel)), 
-						LOCTEXT("PickNewBlueprintTooltip_Create", "Create new Blueprint Class"),
-						FSlateIcon(), 
-						FUIAction(
-							FExecuteAction::CreateStatic( &BlueprintNameEntry::CreateBlueprintClicked, InCreationClass ),
-							FCanExecuteAction::CreateStatic( &BlueprintNameEntry::CanCreateBlueprint )
-						)
-					);
-
-				MenuBuilder.BeginSection( NAME_None, LOCTEXT("PickNewBlueprintFilename", "Blueprint Class Filename") );
-				MenuBuilder.AddWidget( BlueprintNameEntry::MakeBlueprintNameWidget(InCreationClass), FText::GetEmpty() );
-				MenuBuilder.EndSection();
-
-				MenuBuilder.BeginSection( NAME_None, LOCTEXT("PickNewBlueprintPath", "Blueprint Class Path") );
-				MenuBuilder.AddWidget( BlueprintNameEntry::MakeBlueprintPathWidget(), FText::GetEmpty() );
-				MenuBuilder.EndSection();
+				CreateBlueprint(PackageName, InCreationClass);
+				FEditorDirectories::Get().SetLastDirectory(ELastDirectory::NEW_ASSET, PackagePath);
 			}
 		}
 
@@ -1126,7 +993,7 @@ namespace ClassViewer
 		}
 
 		/** Returns TRUE if you can derive a Blueprint */
-		bool CanOpenCreateBlueprintMenu(UClass* InCreationClass)
+		bool CanOpenCreateBlueprintDialog(UClass* InCreationClass)
 		{
 			return !InCreationClass->HasAnyClassFlags(CLASS_Deprecated);
 		}
@@ -1268,6 +1135,67 @@ namespace ClassViewer
 		{
 			ClassHierarchy->UpdateClassInNode(InGeneratedClassPackageName, InNewClass, InNewBluePrint );
 		}
+
+		static TSharedRef<SWidget> CreateMenu(UClass* Class, const bool bIsBlueprint, const bool bHasBlueprint)
+		{
+			// Empty list of commands.
+			TSharedPtr< FUICommandList > Commands;
+
+			const bool bShouldCloseWindowAfterMenuSelection = true;	// Set the menu to automatically close when the user commits to a choice
+			FMenuBuilder MenuBuilder(bShouldCloseWindowAfterMenuSelection, Commands);
+			{
+				if (bIsBlueprint)
+				{
+					TAttribute<FText>::FGetter DynamicTooltipGetter;
+					DynamicTooltipGetter.BindStatic(&ClassViewer::Helpers::GetCreateBlueprintTooltip, Class);
+					TAttribute<FText> DynamicTooltipAttribute = TAttribute<FText>::Create(DynamicTooltipGetter);
+
+					MenuBuilder.AddMenuEntry(
+						LOCTEXT("ClassViewerMenuCreateBlueprint", "Create Blueprint Class..."),
+						DynamicTooltipAttribute,
+						FSlateIcon(),
+						FUIAction(
+							FExecuteAction::CreateStatic(&ClassViewer::Helpers::OpenCreateBlueprintDialog, Class),
+							FCanExecuteAction::CreateStatic(&ClassViewer::Helpers::CanOpenCreateBlueprintDialog, Class)
+							)
+						);
+				}
+
+				if (bHasBlueprint)
+				{
+					MenuBuilder.BeginSection("ClassViewerDropDownHasBlueprint");
+					{
+						FUIAction Action(FExecuteAction::CreateStatic(&ClassViewer::Helpers::OpenBlueprintTool, ClassViewer::Helpers::GetBlueprint(Class)));
+						MenuBuilder.AddMenuEntry(LOCTEXT("ClassViewerMenuEditBlueprint", "Edit Blueprint Class..."), LOCTEXT("ClassViewerMenuEditBlueprint_Tooltip", "Open the Blueprint Class in the editor."), FSlateIcon(), Action);
+					}
+					MenuBuilder.EndSection();
+
+					MenuBuilder.BeginSection("ClassViewerDropDownHasBlueprint2");
+					{
+						FUIAction Action(FExecuteAction::CreateStatic(&ClassViewer::Helpers::FindInContentBrowser, ClassViewer::Helpers::GetBlueprint(Class), Class));
+						MenuBuilder.AddMenuEntry(LOCTEXT("ClassViewerMenuFindContent", "Find in Content Browser..."), LOCTEXT("ClassViewerMenuFindContent_Tooltip", "Find in Content Browser"), FSlateIcon(), Action);
+					}
+					MenuBuilder.EndSection();
+				}
+				else
+				{
+					MenuBuilder.BeginSection("ClassViewerIsCode");
+					{
+						FUIAction Action(FExecuteAction::CreateStatic(&ClassViewer::Helpers::OpenClassHeaderFileInIDE, Class));
+						MenuBuilder.AddMenuEntry(LOCTEXT("ClassViewerMenuOpenCPlusPlusClass", "Open C++ Header..."), LOCTEXT("ClassViewerMenuOpenCPlusPlusClass_Tooltip", "Open the header file for this class in the IDE."), FSlateIcon(), Action);
+					}
+					{
+						FUIAction Action(FExecuteAction::CreateStatic(&ClassViewer::Helpers::OpenCreateCPlusPlusClassWizard, Class));
+						MenuBuilder.AddMenuEntry(LOCTEXT("ClassViewerMenuCreateCPlusPlusClass", "Create New C++ Class..."), LOCTEXT("ClassViewerMenuCreateCPlusPlusClass_Tooltip", "Creates a new C++ class using this class as a base."), FSlateIcon(), Action);
+					}
+					MenuBuilder.EndSection();
+				}
+			}
+
+			return MenuBuilder.MakeWidget();
+		}
+
+
 	} // namespace Helpers
 } // namespace ClassViewer
 
@@ -1468,77 +1396,17 @@ private:
 	 */
 	TSharedRef<SWidget> GenerateDropDown()
 	{
-		// Empty list of commands.
-		TSharedPtr< FUICommandList > Commands;
-
-		const bool bShouldCloseWindowAfterMenuSelection = true;	// Set the menu to automatically close when the user commits to a choice
-		FMenuBuilder MenuBuilder( bShouldCloseWindowAfterMenuSelection, Commands);
-		{ 
-			if (UClass* Class = AssociatedNode->Class.Get())
-			{
-				bool bIsBlueprint(false);
-				bool bHasBlueprint(false);
-
-				ClassViewer::Helpers::GetClassInfo(Class, bIsBlueprint, bHasBlueprint);
-
-				bHasBlueprint = AssociatedNode->Blueprint.IsValid();
-
-				if (bIsBlueprint)
-				{
-					MenuBuilder.BeginSection("ClassViewerIsBlueprint");
-					{
-						TAttribute<FText>::FGetter DynamicTooltipGetter;
-						DynamicTooltipGetter.BindStatic(&ClassViewer::Helpers::GetCreateBlueprintTooltip, Class);
-						TAttribute<FText> DynamicTooltipAttribute = TAttribute<FText>::Create(DynamicTooltipGetter);
-
-						MenuBuilder.AddSubMenu(
-							LOCTEXT("ClassViewerMenuCreateBlueprint", "Create Blueprint Class"), 
-							DynamicTooltipAttribute, 
-							FNewMenuDelegate::CreateStatic( &ClassViewer::Helpers::OpenCreateBlueprintMenu, Class ),
-							FUIAction(
-								FExecuteAction(),
-								FCanExecuteAction::CreateStatic( &ClassViewer::Helpers::CanOpenCreateBlueprintMenu, Class )
-								),
-							FName(),
-							EUserInterfaceActionType::Button
-							);
-					}
-					MenuBuilder.EndSection();
-				}
-
-				if (bHasBlueprint)
-				{
-					MenuBuilder.BeginSection("ClassViewerDropDownHasBlueprint");
-					{
-						FUIAction Action( FExecuteAction::CreateStatic( &ClassViewer::Helpers::OpenBlueprintTool, ClassViewer::Helpers::GetBlueprint(Class) ) );
-						MenuBuilder.AddMenuEntry(LOCTEXT("ClassViewerMenuEditBlueprint", "Edit Blueprint Class..."), LOCTEXT("ClassViewerMenuEditBlueprint_Tooltip", "Open the Blueprint Class in the editor."), FSlateIcon(), Action);
-					}
-					MenuBuilder.EndSection();
+		if (UClass* Class = AssociatedNode->Class.Get())
+		{
+			bool bIsBlueprint(false);
+			bool bHasBlueprint(false);
 
-					MenuBuilder.BeginSection("ClassViewerDropDownHasBlueprint2");
-					{
-						FUIAction Action( FExecuteAction::CreateStatic( &ClassViewer::Helpers::FindInContentBrowser, ClassViewer::Helpers::GetBlueprint(Class), Class ) );
-						MenuBuilder.AddMenuEntry(LOCTEXT("ClassViewerMenuFindContent", "Find in Content Browser..."), LOCTEXT("ClassViewerMenuFindContent_Tooltip", "Find in Content Browser"), FSlateIcon(), Action);
-					}
-					MenuBuilder.EndSection();
-				}
-				else
-				{
-					MenuBuilder.BeginSection("ClassViewerIsCode");
-					{
-						FUIAction Action( FExecuteAction::CreateStatic( &ClassViewer::Helpers::OpenClassHeaderFileInIDE, Class ) );
-						MenuBuilder.AddMenuEntry( LOCTEXT("ClassViewerMenuOpenCPlusPlusClass", "Open C++ Header..."), LOCTEXT("ClassViewerMenuOpenCPlusPlusClass_Tooltip", "Open the header file for this class in the IDE."), FSlateIcon(), Action );
-					}
-					{
-						FUIAction Action( FExecuteAction::CreateStatic( &ClassViewer::Helpers::OpenCreateCPlusPlusClassWizard, Class ) );
-						MenuBuilder.AddMenuEntry( LOCTEXT("ClassViewerMenuCreateCPlusPlusClass", "Create New C++ Class..."), LOCTEXT("ClassViewerMenuCreateCPlusPlusClass_Tooltip", "Creates a new C++ class using this class as a base."), FSlateIcon(), Action );
-					}
-					MenuBuilder.EndSection();
-				}
-			}
+			ClassViewer::Helpers::GetClassInfo(Class, bIsBlueprint, bHasBlueprint);
+			bHasBlueprint = AssociatedNode->Blueprint.IsValid();
+			return ClassViewer::Helpers::CreateMenu(Class, bIsBlueprint, bHasBlueprint);
 		}
 
-		return MenuBuilder.MakeWidget();
+		return SNullWidget::NullWidget;
 	}
 
 	/** Returns the text color for the item based on if it is selected or not. */
@@ -2352,9 +2220,6 @@ void SClassViewer::OnClassViewerExpansionChanged(TSharedPtr<FClassViewerNode> It
 
 TSharedPtr< SWidget > SClassViewer::BuildMenuWidget()
 {
-	// Empty list of commands.
-	TSharedPtr< FUICommandList > Commands;
-
 	bool bIsBlueprint;
 	bool bHasBlueprint;
 	TArray< TSharedPtr< FClassViewerNode > > SelectedList;
@@ -2369,12 +2234,10 @@ TSharedPtr< SWidget > SClassViewer::BuildMenuWidget()
 		SelectedList = ClassList->GetSelectedItems();
 	}
 
-	// If there is no selected item, build an empty menu and return.
+	// If there is no selected item, return a null widget.
 	if(SelectedList.Num() == 0)
 	{
-		const bool bShouldCloseWindowAfterMenuSelection = true;	// Set the menu to automatically close when the user commits to a choice
-		FMenuBuilder MenuBuilder( bShouldCloseWindowAfterMenuSelection, Commands);
-		return MenuBuilder.MakeWidget();
+		return SNullWidget::NullWidget;
 	}
 
 	// If it is NOT stale, it has not been set (meaning it was never valid but now is invalid).
@@ -2396,58 +2259,7 @@ TSharedPtr< SWidget > SClassViewer::BuildMenuWidget()
 		bHasBlueprint = true;
 	}
 
-	const bool bShouldCloseWindowAfterMenuSelection = true;	// Set the menu to automatically close when the user commits to a choice
-	FMenuBuilder MenuBuilder( bShouldCloseWindowAfterMenuSelection, Commands);
-	{ 
-		if (bIsBlueprint)
-		{
-			TAttribute<FText>::FGetter DynamicTooltipGetter;
-			DynamicTooltipGetter.BindStatic(&ClassViewer::Helpers::GetCreateBlueprintTooltip, RightClickClass);
-			TAttribute<FText> DynamicTooltipAttribute = TAttribute<FText>::Create(DynamicTooltipGetter);
-
-			MenuBuilder.AddSubMenu(
-				LOCTEXT("ClassViewerMenuCreateBlueprint", "Create Blueprint Class"), 
-				DynamicTooltipAttribute, 
-				FNewMenuDelegate::CreateStatic( &ClassViewer::Helpers::OpenCreateBlueprintMenu, RightClickClass ),
-				FUIAction(
-					FExecuteAction(),
-					FCanExecuteAction::CreateStatic( &ClassViewer::Helpers::CanOpenCreateBlueprintMenu, RightClickClass )
-					),
-				FName(),
-				EUserInterfaceActionType::Button
-				);
-		}
-	
-		if(bHasBlueprint)
-		{
-			MenuBuilder.BeginSection("ClassViewerHasBlueprint");
-			{
-				FUIAction Action( FExecuteAction::CreateRaw( this, &SClassViewer::OnOpenBlueprintTool ) );
-				MenuBuilder.AddMenuEntry(LOCTEXT("ClassViewerMenuEditBlueprint", "Edit Blueprint Class..."), LOCTEXT("ClassViewerMenuEditBlueprint_Tooltip", "Open the Blueprint Class in the editor."), FSlateIcon(), Action);
-			}
-			MenuBuilder.EndSection();
-
-			MenuBuilder.BeginSection("ClassViewerHasBlueprint2");
-			{
-				FUIAction Action( FExecuteAction::CreateRaw( this, &SClassViewer::FindInContentBrowser ) );
-				MenuBuilder.AddMenuEntry(LOCTEXT("ClassViewerMenuFindContent", "Find in Content Browser..."), LOCTEXT("ClassViewerMenuFindContent_Tooltip", "Find in Content Browser"), FSlateIcon(), Action);
-			}
-			MenuBuilder.EndSection();
-		}
-		else
-		{
-			{
-				FUIAction Action( FExecuteAction::CreateStatic( &ClassViewer::Helpers::OpenClassHeaderFileInIDE, RightClickClass ) );
-				MenuBuilder.AddMenuEntry( LOCTEXT("ClassViewerMenuOpenCPlusPlusClass", "Open C++ Header..."), LOCTEXT("ClassViewerMenuOpenCPlusPlusClass_Tooltip", "Open the header file for this class in the IDE."), FSlateIcon(), Action );
-			}
-			{
-				FUIAction Action( FExecuteAction::CreateStatic( &ClassViewer::Helpers::OpenCreateCPlusPlusClassWizard, RightClickClass ) );
-				MenuBuilder.AddMenuEntry( LOCTEXT("ClassViewerMenuCreateCPlusPlusClass", "Create New C++ Class..."), LOCTEXT("ClassViewerMenuCreateCPlusPlusClass_Tooltip", "Creates a new C++ class using this class as a base."), FSlateIcon(), Action );
-			}
-		}
-	}
-
-	return MenuBuilder.MakeWidget();
+	return ClassViewer::Helpers::CreateMenu(RightClickClass, bIsBlueprint, bHasBlueprint);
 }
 
 TSharedRef< ITableRow > SClassViewer::OnGenerateRowForClassViewer( TSharedPtr<FClassViewerNode> Item, const TSharedRef< STableViewBase >& OwnerTable )
