Changed SClassItem parent from SComboRow to STableRow so that dragging can be initiated from a selected item.

diff --git a/Engine/Source/Editor/ClassViewer/Private/SClassViewer.cpp b/Engine/Source/Editor/ClassViewer/Private/SClassViewer.cpp
--- a/Engine/Source/Editor/ClassViewer/Private/SClassViewer.cpp
+++ b/Engine/Source/Editor/ClassViewer/Private/SClassViewer.cpp
@@ -1205,7 +1205,7 @@ namespace ClassViewer
 DECLARE_DELEGATE_OneParam( FOnClassItemDoubleClickDelegate, TSharedPtr<FClassViewerNode> );
 
 /** The item used for visualizing the class in the tree. */
-class SClassItem : public SComboRow< TSharedPtr<FString> >
+class SClassItem : public STableRow< TSharedPtr<FString> >
 {
 public:
 	
