Copying //UE4/Release-Staging-4.12 to //UE4/Dev-Main (Source: //UE4/Release-4.12 @ 2955635)

diff --git a/Engine/Source/Editor/AnimGraph/Private/AnimPreviewInstance.cpp b/Engine/Source/Editor/AnimGraph/Private/AnimPreviewInstance.cpp
--- a/Engine/Source/Editor/AnimGraph/Private/AnimPreviewInstance.cpp
+++ b/Engine/Source/Editor/AnimGraph/Private/AnimPreviewInstance.cpp
@@ -191,7 +191,7 @@ void FAnimPreviewInstanceProxy::RefreshCurveBoneControllers()
 {
 	// go through all curves and see if it has Transform Curve
 	// if so, find what bone that belong to and create BoneMOdifier for them
-	UAnimSequence* CurrentSequence = Cast<UAnimSequence>(GetCurrentAsset());
+	UAnimSequence* CurrentSequence = Cast<UAnimSequence>(CurrentAsset);
 
 	CurveBoneControllers.Empty();
 
@@ -243,7 +243,7 @@ void FAnimPreviewInstanceProxy::RefreshCurveBoneControllers()
 void FAnimPreviewInstanceProxy::UpdateCurveController()
 {
 	// evaluate the curve data first
-	UAnimSequenceBase* CurrentSequence = Cast<UAnimSequenceBase>(GetCurrentAsset());
+	UAnimSequenceBase* CurrentSequence = Cast<UAnimSequenceBase>(CurrentAsset);
 
 	if (CurrentSequence && GetSkeleton())
 	{
@@ -302,7 +302,7 @@ void FAnimPreviewInstanceProxy::SetKeyImplementation(const FCompactPose& PreCont
 {
 #if WITH_EDITOR
 	// evaluate the curve data first
-	UAnimSequence* CurrentSequence = Cast<UAnimSequence>(GetCurrentAsset());
+	UAnimSequence* CurrentSequence = Cast<UAnimSequence>(CurrentAsset);
 	UDebugSkelMeshComponent* Component = Cast<UDebugSkelMeshComponent> (GetSkelMeshComponent());
 
 	if(CurrentSequence && GetSkeleton() && Component && Component->SkeletalMesh)
@@ -400,7 +400,7 @@ void UAnimPreviewInstance::NativeInitializeAnimation()
 	FAnimPreviewInstanceProxy& Proxy = GetProxyOnGameThread<FAnimPreviewInstanceProxy>();
 
 	// Cache our play state from the previous animation otherwise set to play
-	bool bCachedIsPlaying = (Proxy.GetCurrentAsset() != NULL) ? Proxy.IsPlaying() : true;
+	bool bCachedIsPlaying = (CurrentAsset != nullptr) ? Proxy.IsPlaying() : true;
 
 	Super::NativeInitializeAnimation();
 
@@ -447,15 +447,15 @@ void UAnimPreviewInstance::SetSkeletalControlAlpha(float InSkeletalControlAlpha)
 
 UAnimSequence* UAnimPreviewInstance::GetAnimSequence()
 {
-	return Cast<UAnimSequence>(GetCurrentAsset());
+	return Cast<UAnimSequence>(CurrentAsset);
 }
 
 void UAnimPreviewInstance::RestartMontage(UAnimMontage* Montage, FName FromSection)
 {
-	FAnimPreviewInstanceProxy& Proxy = GetProxyOnGameThread<FAnimPreviewInstanceProxy>();
-
-	if (Montage == Proxy.GetCurrentAsset())
+	if (Montage == CurrentAsset)
 	{
+		FAnimPreviewInstanceProxy& Proxy = GetProxyOnGameThread<FAnimPreviewInstanceProxy>();
+
 		MontagePreviewType = EMPT_Normal;
 		// since this is preview, we would like not to blend in
 		// just hard stop here
@@ -477,7 +477,7 @@ void UAnimPreviewInstance::SetAnimationAsset(UAnimationAsset* NewAsset, bool bIs
 	Proxy.GetRequiredBones().SetUseSourceData(false);
 
 	Super::SetAnimationAsset(NewAsset, bIsLooping, InPlayRate);
-	RootMotionMode = Cast<UAnimMontage>(Proxy.GetCurrentAsset()) != NULL ? ERootMotionMode::RootMotionFromMontagesOnly : ERootMotionMode::RootMotionFromEverything;
+	RootMotionMode = Cast<UAnimMontage>(CurrentAsset) != nullptr ? ERootMotionMode::RootMotionFromMontagesOnly : ERootMotionMode::RootMotionFromEverything;
 
 	// should re sync up curve bone controllers from new asset
 	Proxy.RefreshCurveBoneControllers();
@@ -488,7 +488,7 @@ void UAnimPreviewInstance::MontagePreview_SetLooping(bool bIsLooping)
 	FAnimPreviewInstanceProxy& Proxy = GetProxyOnGameThread<FAnimPreviewInstanceProxy>();
 	Proxy.SetLooping(bIsLooping);
 
-	if (UAnimMontage* Montage = Cast<UAnimMontage>(Proxy.GetCurrentAsset()))
+	if (UAnimMontage* Montage = Cast<UAnimMontage>(CurrentAsset))
 	{
 		switch (MontagePreviewType)
 		{
@@ -514,7 +514,7 @@ void UAnimPreviewInstance::MontagePreview_SetPlaying(bool bIsPlaying)
 	}
 	else if (Proxy.IsPlaying())
 	{
-		UAnimMontage* Montage = Cast<UAnimMontage>(Proxy.GetCurrentAsset());
+		UAnimMontage* Montage = Cast<UAnimMontage>(CurrentAsset);
 		if (Montage)
 		{
 			switch (MontagePreviewType)
@@ -546,7 +546,7 @@ void UAnimPreviewInstance::MontagePreview_SetReverse(bool bInReverse)
 void UAnimPreviewInstance::MontagePreview_Restart()
 {
 	FAnimPreviewInstanceProxy& Proxy = GetProxyOnGameThread<FAnimPreviewInstanceProxy>();
-	if (UAnimMontage* Montage = Cast<UAnimMontage>(Proxy.GetCurrentAsset()))
+	if (UAnimMontage* Montage = Cast<UAnimMontage>(CurrentAsset))
 	{
 		switch (MontagePreviewType)
 		{
@@ -565,7 +565,7 @@ void UAnimPreviewInstance::MontagePreview_StepForward()
 {
 	FAnimPreviewInstanceProxy& Proxy = GetProxyOnGameThread<FAnimPreviewInstanceProxy>();
 
-	if (UAnimMontage* Montage = Cast<UAnimMontage>(Proxy.GetCurrentAsset()))
+	if (UAnimMontage* Montage = Cast<UAnimMontage>(CurrentAsset))
 	{
 		bool bWasPlaying = IsPlayingMontage() && (Proxy.IsLooping() || Proxy.IsPlaying()); // we need to handle non-looped case separately, even if paused during playthrough
 		MontagePreview_SetReverse(false);
@@ -615,10 +615,9 @@ void UAnimPreviewInstance::MontagePreview_StepForward()
 
 void UAnimPreviewInstance::MontagePreview_StepBackward()
 {
-	FAnimPreviewInstanceProxy& Proxy = GetProxyOnGameThread<FAnimPreviewInstanceProxy>();
-
-	if (UAnimMontage* Montage = Cast<UAnimMontage>(Proxy.GetCurrentAsset()))
+	if (UAnimMontage* Montage = Cast<UAnimMontage>(CurrentAsset))
 	{
+		FAnimPreviewInstanceProxy& Proxy = GetProxyOnGameThread<FAnimPreviewInstanceProxy>();
 		bool bWasPlaying = IsPlayingMontage() && (Proxy.IsLooping() || Proxy.IsPlaying()); // we need to handle non-looped case separately, even if paused during playthrough
 		MontagePreview_SetReverse(true);
 		if (! bWasPlaying)
@@ -677,9 +676,10 @@ float UAnimPreviewInstance::MontagePreview_CalculateStepLength()
 
 void UAnimPreviewInstance::MontagePreview_JumpToStart()
 {
-	FAnimPreviewInstanceProxy& Proxy = GetProxyOnGameThread<FAnimPreviewInstanceProxy>();
-	if (UAnimMontage* Montage = Cast<UAnimMontage>(Proxy.GetCurrentAsset()))
+	if (UAnimMontage* Montage = Cast<UAnimMontage>(CurrentAsset))
 	{
+		FAnimPreviewInstanceProxy& Proxy = GetProxyOnGameThread<FAnimPreviewInstanceProxy>();
+
 		int32 SectionIdx = 0;
 		if (MontagePreviewType == EMPT_Normal)
 		{
@@ -708,9 +708,10 @@ void UAnimPreviewInstance::MontagePreview_JumpToStart()
 
 void UAnimPreviewInstance::MontagePreview_JumpToEnd()
 {
-	FAnimPreviewInstanceProxy& Proxy = GetProxyOnGameThread<FAnimPreviewInstanceProxy>();
-	if (UAnimMontage* Montage = Cast<UAnimMontage>(Proxy.GetCurrentAsset()))
+	if (UAnimMontage* Montage = Cast<UAnimMontage>(CurrentAsset))
 	{
+		FAnimPreviewInstanceProxy& Proxy = GetProxyOnGameThread<FAnimPreviewInstanceProxy>();
+
 		int32 SectionIdx = 0;
 		if (MontagePreviewType == EMPT_Normal)
 		{
@@ -739,10 +740,10 @@ void UAnimPreviewInstance::MontagePreview_JumpToEnd()
 
 void UAnimPreviewInstance::MontagePreview_JumpToPreviewStart()
 {
-	FAnimPreviewInstanceProxy& Proxy = GetProxyOnGameThread<FAnimPreviewInstanceProxy>();
-
-	if (UAnimMontage* Montage = Cast<UAnimMontage>(Proxy.GetCurrentAsset()))
+	if (UAnimMontage* Montage = Cast<UAnimMontage>(CurrentAsset))
 	{
+		FAnimPreviewInstanceProxy& Proxy = GetProxyOnGameThread<FAnimPreviewInstanceProxy>();
+
 		int32 SectionIdx = 0;
 		if (MontagePreviewType == EMPT_Normal)
 		{
@@ -764,11 +765,11 @@ void UAnimPreviewInstance::MontagePreview_JumpToPreviewStart()
 
 void UAnimPreviewInstance::MontagePreview_JumpToPosition(float NewPosition)
 {
-	FAnimPreviewInstanceProxy& Proxy = GetProxyOnGameThread<FAnimPreviewInstanceProxy>();
-
 	SetPosition(NewPosition, false);
-	if (UAnimMontage* Montage = Cast<UAnimMontage>(Proxy.GetCurrentAsset()))
+	if (UAnimMontage* Montage = Cast<UAnimMontage>(CurrentAsset))
 	{
+		FAnimPreviewInstanceProxy& Proxy = GetProxyOnGameThread<FAnimPreviewInstanceProxy>();
+
 		// this section will be first
 		int32 NewMontagePreviewStartSectionIdx = MontagePreview_FindFirstSectionAsInMontage(Montage->GetSectionIndexFromPosition(NewPosition));
 		if (MontagePreviewStartSectionIdx != NewMontagePreviewStartSectionIdx &&
@@ -791,10 +792,10 @@ void UAnimPreviewInstance::MontagePreview_RemoveBlendOut()
 
 void UAnimPreviewInstance::MontagePreview_PreviewNormal(int32 FromSectionIdx)
 {
-	FAnimPreviewInstanceProxy& Proxy = GetProxyOnGameThread<FAnimPreviewInstanceProxy>();
-
-	if (UAnimMontage* Montage = Cast<UAnimMontage>(Proxy.GetCurrentAsset()))
+	if (UAnimMontage* Montage = Cast<UAnimMontage>(CurrentAsset))
 	{
+		FAnimPreviewInstanceProxy& Proxy = GetProxyOnGameThread<FAnimPreviewInstanceProxy>();
+
 		int32 PreviewFromSection = FromSectionIdx;
 		if (FromSectionIdx != INDEX_NONE)
 		{
@@ -819,11 +820,11 @@ void UAnimPreviewInstance::MontagePreview_PreviewNormal(int32 FromSectionIdx)
 
 void UAnimPreviewInstance::MontagePreview_PreviewAllSections()
 {
-	FAnimPreviewInstanceProxy& Proxy = GetProxyOnGameThread<FAnimPreviewInstanceProxy>();
-
-	UAnimMontage* Montage = Cast<UAnimMontage>(Proxy.GetCurrentAsset());
+	UAnimMontage* Montage = Cast<UAnimMontage>(CurrentAsset);
 	if (Montage && Montage->SequenceLength > 0.f)
 	{
+		FAnimPreviewInstanceProxy& Proxy = GetProxyOnGameThread<FAnimPreviewInstanceProxy>();
+
 		MontagePreviewType = EMPT_AllSections;
 		// since this is preview, we would like not to blend in
 		// just hard stop here
@@ -838,10 +839,10 @@ void UAnimPreviewInstance::MontagePreview_PreviewAllSections()
 
 void UAnimPreviewInstance::MontagePreview_SetLoopNormal(bool bIsLooping, int32 PreferSectionIdx)
 {
-	FAnimPreviewInstanceProxy& Proxy = GetProxyOnGameThread<FAnimPreviewInstanceProxy>();
-
-	if (UAnimMontage* Montage = Cast<UAnimMontage>(Proxy.GetCurrentAsset()))
+	if (UAnimMontage* Montage = Cast<UAnimMontage>(CurrentAsset))
 	{
+		FAnimPreviewInstanceProxy& Proxy = GetProxyOnGameThread<FAnimPreviewInstanceProxy>();
+
 		MontagePreview_ResetSectionsOrder();
 
 		if (PreferSectionIdx == INDEX_NONE)
@@ -904,10 +905,10 @@ void UAnimPreviewInstance::MontagePreview_SetLoopNormal(bool bIsLooping, int32 P
 
 void UAnimPreviewInstance::MontagePreview_SetLoopAllSetupSections(bool bIsLooping)
 {
-	FAnimPreviewInstanceProxy& Proxy = GetProxyOnGameThread<FAnimPreviewInstanceProxy>();
-
-	if (UAnimMontage* Montage = Cast<UAnimMontage>(Proxy.GetCurrentAsset()))
+	if (UAnimMontage* Montage = Cast<UAnimMontage>(CurrentAsset))
 	{
+		FAnimPreviewInstanceProxy& Proxy = GetProxyOnGameThread<FAnimPreviewInstanceProxy>();
+
 		MontagePreview_ResetSectionsOrder();
 
 		int32 TotalSection = Montage->CompositeSections.Num();
@@ -961,10 +962,10 @@ void UAnimPreviewInstance::MontagePreview_SetLoopAllSetupSections(bool bIsLoopin
 
 void UAnimPreviewInstance::MontagePreview_SetLoopAllSections(bool bIsLooping)
 {
-	FAnimPreviewInstanceProxy& Proxy = GetProxyOnGameThread<FAnimPreviewInstanceProxy>();
-
-	if (UAnimMontage* Montage = Cast<UAnimMontage>(Proxy.GetCurrentAsset()))
+	if (UAnimMontage* Montage = Cast<UAnimMontage>(CurrentAsset))
 	{
+		FAnimPreviewInstanceProxy& Proxy = GetProxyOnGameThread<FAnimPreviewInstanceProxy>();
+
 		int32 TotalSection = Montage->CompositeSections.Num();
 		if (TotalSection > 0)
 		{
@@ -989,8 +990,7 @@ void UAnimPreviewInstance::MontagePreview_SetLoopAllSections(bool bIsLooping)
 
 void UAnimPreviewInstance::MontagePreview_ResetSectionsOrder()
 {
-	FAnimPreviewInstanceProxy& Proxy = GetProxyOnGameThread<FAnimPreviewInstanceProxy>();
-	if (UAnimMontage* Montage = Cast<UAnimMontage>(Proxy.GetCurrentAsset()))
+	if (UAnimMontage* Montage = Cast<UAnimMontage>(CurrentAsset))
 	{
 		int32 TotalSection = Montage->CompositeSections.Num();
 		// restore to default
@@ -1003,10 +1003,9 @@ void UAnimPreviewInstance::MontagePreview_ResetSectionsOrder()
 
 int32 UAnimPreviewInstance::MontagePreview_FindFirstSectionAsInMontage(int32 ForSectionIdx)
 {
-	FAnimPreviewInstanceProxy& Proxy = GetProxyOnGameThread<FAnimPreviewInstanceProxy>();
 	int32 ResultIdx = ForSectionIdx;
 	// Montage does not have looping set up, so it should be valid and it gets
-	if (UAnimMontage* Montage = Cast<UAnimMontage>(Proxy.GetCurrentAsset()))
+	if (UAnimMontage* Montage = Cast<UAnimMontage>(CurrentAsset))
 	{
 		TArray<bool> AlreadyVisited;
 		AlreadyVisited.AddZeroed(Montage->CompositeSections.Num());
@@ -1047,9 +1046,8 @@ int32 UAnimPreviewInstance::MontagePreview_FindFirstSectionAsInMontage(int32 For
 
 int32 UAnimPreviewInstance::MontagePreview_FindLastSection(int32 StartSectionIdx)
 {
-	FAnimPreviewInstanceProxy& Proxy = GetProxyOnGameThread<FAnimPreviewInstanceProxy>();
 	int32 ResultIdx = StartSectionIdx;
-	if (UAnimMontage* Montage = Cast<UAnimMontage>(Proxy.GetCurrentAsset()))
+	if (UAnimMontage* Montage = Cast<UAnimMontage>(CurrentAsset))
 	{
 		if (FAnimMontageInstance* CurMontageInstance = GetActiveMontageInstance())
 		{
diff --git a/Engine/Source/Editor/AnimGraph/Private/AnimStateNodeBase.cpp b/Engine/Source/Editor/AnimGraph/Private/AnimStateNodeBase.cpp
--- a/Engine/Source/Editor/AnimGraph/Private/AnimStateNodeBase.cpp
+++ b/Engine/Source/Editor/AnimGraph/Private/AnimStateNodeBase.cpp
@@ -10,6 +10,7 @@
 #include "EdGraphUtilities.h"
 #include "Kismet2/Kismet2NameValidators.h"
 #include "AnimStateNodeBase.h"
+#include "FrameworkObjectVersion.h"
 /////////////////////////////////////////////////////
 // FAnimStateNodeNameValidator
 
@@ -93,6 +94,37 @@ FString UAnimStateNodeBase::GetDocumentationLink() const
 	return TEXT("Shared/GraphNodes/AnimationStateMachine");
 }
 
+void UAnimStateNodeBase::Serialize(FArchive& Ar)
+{
+	Super::Serialize(Ar);
+	Ar.UsingCustomVersion(FFrameworkObjectVersion::GUID);
+}
+
+void UAnimStateNodeBase::PostLoad()
+{
+	Super::PostLoad();
+
+	const int32 CustomVersion = GetLinkerCustomVersion(FFrameworkObjectVersion::GUID);
+
+	if(CustomVersion < FFrameworkObjectVersion::FixNonTransactionalPins)
+	{
+		int32 BrokenPinCount = 0;
+		for(UEdGraphPin* Pin : Pins)
+		{
+			if(!Pin->HasAnyFlags(RF_Transactional))
+			{
+				++BrokenPinCount;
+				Pin->SetFlags(Pin->GetFlags() | RF_Transactional);
+			}
+		}
+
+		if(BrokenPinCount > 0)
+		{
+			UE_LOG(LogAnimation, Log, TEXT("Fixed %d non-transactional pins in %s"), BrokenPinCount, *GetName());
+		}
+	}
+}
+
 UAnimBlueprint* UAnimStateNodeBase::GetAnimBlueprint() const
 {
 	UBlueprint* Blueprint = FBlueprintEditorUtils::FindBlueprintForNode(this);
diff --git a/Engine/Source/Editor/BlueprintGraph/Classes/K2Node_GetClassDefaults.h b/Engine/Source/Editor/BlueprintGraph/Classes/K2Node_GetClassDefaults.h
--- a/Engine/Source/Editor/BlueprintGraph/Classes/K2Node_GetClassDefaults.h
+++ b/Engine/Source/Editor/BlueprintGraph/Classes/K2Node_GetClassDefaults.h
@@ -35,11 +35,17 @@ class UK2Node_GetClassDefaults : public UK2Node
 
 public:
 	/** Finds and returns the class input pin from the current set of pins. */
-	UEdGraphPin* FindClassPin() const
+	BLUEPRINTGRAPH_API UEdGraphPin* FindClassPin() const
 	{
 		return FindClassPin(Pins);
 	}
 
+	/** Retrieves the current input class type. */
+	BLUEPRINTGRAPH_API UClass* GetInputClass() const
+	{
+		return GetInputClass(FindClassPin());
+	}
+
 protected:
 	/**
 	 * Finds and returns the class input pin.
@@ -49,11 +55,11 @@ protected:
 	UEdGraphPin* FindClassPin(const TArray<UEdGraphPin*>& FromPins) const;
 
 	/**
-	 * Determines the input class type from the given pin.
-	 *
-	 * @param FromPin	Input class pin. If not set (default), then it will fall back to using FindClassPin().
-	 */
-	UClass* GetInputClass(const UEdGraphPin* FromPin = nullptr) const;
+	* Determines the input class type from the given pin.
+	*
+	* @param FromPin	Input class pin.
+	*/
+	UClass* GetInputClass(const UEdGraphPin* FromPin) const;
 
 	/**
 	 * Creates the full set of output pins (properties) from the given input class.
@@ -69,6 +75,12 @@ private:
 	/** Class pin name */
 	static FString ClassPinName;
 
+	/** Blueprint.OnChanged delegate handle */
+	FDelegateHandle OnBlueprintChangedDelegate;
+
+	/** Blueprint.OnCompiled delegate handle */
+	FDelegateHandle OnBlueprintCompiledDelegate;
+
 	/** Output pin visibility control */
 	UPROPERTY(EditAnywhere, Category=PinOptions, EditFixedSize)
 	TArray<FOptionalPinFromProperty> ShowPinForProperties;
diff --git a/Engine/Source/Editor/BlueprintGraph/Private/DelegateNodeHandlers.cpp b/Engine/Source/Editor/BlueprintGraph/Private/DelegateNodeHandlers.cpp
--- a/Engine/Source/Editor/BlueprintGraph/Private/DelegateNodeHandlers.cpp
+++ b/Engine/Source/Editor/BlueprintGraph/Private/DelegateNodeHandlers.cpp
@@ -194,14 +194,14 @@ void FKCHandler_CreateDelegate::RegisterNets(FKismetFunctionContext& Context, UE
 	const FName DelegateFunctionName = DelegateNode->GetFunctionName();
 	if(DelegateFunctionName == NAME_None)
 	{
-		CompilerContext.MessageLog.Error(*LOCTEXT("NoDelegateFunctionName", "No delegate function name @@").ToString(), DelegateNode);
+		CompilerContext.MessageLog.Error(*LOCTEXT("NoDelegateFunctionName", "@@ : missing a function/event name.").ToString(), DelegateNode);
 		return;
 	}
 
 	if(!DelegateNode->GetDelegateSignature())
 	{
 		const FString ErrorStr = FString::Printf(
-			*LOCTEXT("NoDelegateFunction", "No delegate function '%' @@").ToString(),
+			*LOCTEXT("NoDelegateFunction", "@@ : unable to determine expected signature - is the delegate pin connected?").ToString(),
 			*DelegateFunctionName.ToString());
 		CompilerContext.MessageLog.Error(*ErrorStr, DelegateNode);
 		return;
diff --git a/Engine/Source/Editor/BlueprintGraph/Private/K2Node_CallFunction.cpp b/Engine/Source/Editor/BlueprintGraph/Private/K2Node_CallFunction.cpp
--- a/Engine/Source/Editor/BlueprintGraph/Private/K2Node_CallFunction.cpp
+++ b/Engine/Source/Editor/BlueprintGraph/Private/K2Node_CallFunction.cpp
@@ -716,8 +716,12 @@ void UK2Node_CallFunction::CreateExecPinsForFunctionCall(const UFunction* Functi
 				int32 NumExecs = (EnumProp->Enum->NumEnums() - 1);
 				for(int32 ExecIdx=0; ExecIdx<NumExecs; ExecIdx++)
 				{
-					FString ExecName = EnumProp->Enum->GetEnumName(ExecIdx);
-					CreatePin(Direction, K2Schema->PC_Exec, TEXT(""), NULL, false, false, ExecName);
+					bool const bShouldBeHidden = EnumProp->Enum->HasMetaData(TEXT("Hidden"), ExecIdx) || EnumProp->Enum->HasMetaData(TEXT("Spacer"), ExecIdx);
+					if (!bShouldBeHidden)
+					{
+						FString ExecName = EnumProp->Enum->GetEnumName(ExecIdx);
+						CreatePin(Direction, K2Schema->PC_Exec, TEXT(""), NULL, false, false, ExecName);
+					}
 				}
 				
 				if (bIsFunctionInput)
diff --git a/Engine/Source/Editor/BlueprintGraph/Private/K2Node_ConstructObjectFromClass.cpp b/Engine/Source/Editor/BlueprintGraph/Private/K2Node_ConstructObjectFromClass.cpp
--- a/Engine/Source/Editor/BlueprintGraph/Private/K2Node_ConstructObjectFromClass.cpp
+++ b/Engine/Source/Editor/BlueprintGraph/Private/K2Node_ConstructObjectFromClass.cpp
@@ -68,7 +68,7 @@ void UK2Node_ConstructObjectFromClass::AllocateDefaultPins()
 
 UEdGraphPin* UK2Node_ConstructObjectFromClass::GetOuterPin() const
 {
-	UEdGraphPin* Pin = FindPinChecked(FK2Node_ConstructObjectFromClassHelper::OuterPinName);
+	UEdGraphPin* Pin = FindPin(FK2Node_ConstructObjectFromClassHelper::OuterPinName);
 	ensure(nullptr == Pin || Pin->Direction == EGPD_Input);
 	return Pin;
 }
@@ -232,11 +232,20 @@ void UK2Node_ConstructObjectFromClass::PinConnectionListChanged(UEdGraphPin* Pin
 
 void UK2Node_ConstructObjectFromClass::GetPinHoverText(const UEdGraphPin& Pin, FString& HoverTextOut) const
 {
-	SetPinToolTip(*GetClassPin(), LOCTEXT("ClassPinDescription", "The object class you want to construct"));
-	SetPinToolTip(*GetResultPin(), LOCTEXT("ResultPinDescription", "The constructed object"));
-	if (UseOuter())
+	UEdGraphPin* ClassPin = GetClassPin();
+	if (ClassPin)
+	{
+		SetPinToolTip(*ClassPin, LOCTEXT("ClassPinDescription", "The object class you want to construct"));
+	}
+	UEdGraphPin* ResultPin = GetResultPin();
+	if (ResultPin)
+	{
+		SetPinToolTip(*ResultPin, LOCTEXT("ResultPinDescription", "The constructed object"));
+	}
+	UEdGraphPin* OuterPin = UseOuter() ? GetOuterPin() : nullptr;
+	if (OuterPin)
 	{
-		SetPinToolTip(*GetOuterPin(), LOCTEXT("OuterPinDescription", "Owner of the constructed object"));
+		SetPinToolTip(*OuterPin, LOCTEXT("OuterPinDescription", "Owner of the constructed object"));
 	}
 
 	return Super::GetPinHoverText(Pin, HoverTextOut);
diff --git a/Engine/Source/Editor/BlueprintGraph/Private/K2Node_CreateDelegate.cpp b/Engine/Source/Editor/BlueprintGraph/Private/K2Node_CreateDelegate.cpp
--- a/Engine/Source/Editor/BlueprintGraph/Private/K2Node_CreateDelegate.cpp
+++ b/Engine/Source/Editor/BlueprintGraph/Private/K2Node_CreateDelegate.cpp
@@ -8,7 +8,6 @@
 
 struct FK2Node_CreateDelegate_Helper
 {
-	static FString ObjectInputName;
 	static FString DelegateOutputName;
 };
 FString FK2Node_CreateDelegate_Helper::DelegateOutputName(TEXT("OutputDelegate"));
@@ -37,12 +36,12 @@ void UK2Node_CreateDelegate::AllocateDefaultPins()
 
 bool UK2Node_CreateDelegate::IsValid(FString* OutMsg, bool bDontUseSkeletalClassForSelf) const
 {
-
-	if (GetFunctionName() == NAME_None)
+	FName FunctionName = GetFunctionName();
+	if (FunctionName == NAME_None)
 	{
 		if (OutMsg)
 		{
-			*OutMsg = NSLOCTEXT("K2Node", "No_function_name", "No function name.").ToString();
+			*OutMsg = NSLOCTEXT("K2Node", "No_function_name", "No function/event specified.").ToString();
 		}
 		return false;
 	}
@@ -52,7 +51,7 @@ bool UK2Node_CreateDelegate::IsValid(FString* OutMsg, bool bDontUseSkeletalClass
 	{
 		if (OutMsg)
 		{
-			*OutMsg = NSLOCTEXT("K2Node", "No_delegate_out_pin", "No delegate out pin.").ToString();
+			*OutMsg = NSLOCTEXT("K2Node", "No_delegate_out_pin", "Malformed node - there's no delegate output pin.").ToString();
 		}
 		return false;
 	}
@@ -62,7 +61,7 @@ bool UK2Node_CreateDelegate::IsValid(FString* OutMsg, bool bDontUseSkeletalClass
 	{
 		if (OutMsg)
 		{
-			*OutMsg = NSLOCTEXT("K2Node", "Signature_not_found", "Signature not found.").ToString();
+			*OutMsg = NSLOCTEXT("K2Node", "Signature_not_found", "Unable to determine expected signature - is the delegate pin connected?").ToString();
 		}
 		return false;
 	}
@@ -76,7 +75,17 @@ bool UK2Node_CreateDelegate::IsValid(FString* OutMsg, bool bDontUseSkeletalClass
 		{
 			if (OutMsg)
 			{
-				*OutMsg = NSLOCTEXT("K2Node", "No_delegate_out_pin", "No delegate out pin.").ToString();
+				if (const UK2Node_BaseMCDelegate* DelegateNode = Cast<const UK2Node_BaseMCDelegate>(OtherPin->GetOwningNode()))
+				{
+					const FString DelegateName = DelegateNode->GetPropertyName().ToString();
+
+					*OutMsg = FString::Printf(*NSLOCTEXT("K2Node", "Bad_delegate_connection_named", "A connected delegate (%s) has an incompatible signature - has that delegate changed?").ToString(),
+						*DelegateName);
+				}
+				else
+				{
+					*OutMsg = NSLOCTEXT("K2Node", "Bad_delegate_connection", "A connected delegate's signature is incompatible - has that delegate changed?").ToString();
+				}
 			}
 			return false;
 		}
@@ -87,7 +96,14 @@ bool UK2Node_CreateDelegate::IsValid(FString* OutMsg, bool bDontUseSkeletalClass
 	{
 		if (OutMsg)
 		{
-			*OutMsg = NSLOCTEXT("K2Node", "Class_not_found", "Class not found.").ToString();
+			FString SelfPinName = UEdGraphSchema_K2::PN_Self;
+			if (UEdGraphPin* SelfPin = GetObjectInPin())
+			{
+				SelfPinName = SelfPin->PinFriendlyName.IsEmpty() ? SelfPin->PinFriendlyName.ToString() : SelfPin->PinName;
+			}
+
+			*OutMsg = FString::Printf(*NSLOCTEXT("K2Node", "Class_not_found", "Unable to determine context for the selected function/event: '%s' - make sure the target '%s' pin is properly set up.").ToString(),
+				*FunctionName.ToString(), *SelfPinName);
 		}
 		return false;
 	}
@@ -99,7 +115,9 @@ bool UK2Node_CreateDelegate::IsValid(FString* OutMsg, bool bDontUseSkeletalClass
 	{
 		if (OutMsg)
 		{
-			*OutMsg = NSLOCTEXT("K2Node", "Function_not_found", "Function not found.").ToString();
+			*OutMsg = FString::Printf(*NSLOCTEXT("K2Node", "Function_not_found", "Unable to find the selected function/event: '%s' - has it been deleted?").ToString(),
+				*FunctionName.ToString());
+
 		}
 		return false;
 	}
@@ -107,7 +125,8 @@ bool UK2Node_CreateDelegate::IsValid(FString* OutMsg, bool bDontUseSkeletalClass
 	{
 		if (OutMsg)
 		{
-			*OutMsg = NSLOCTEXT("K2Node", "Function_not_compatible", "Function not compatible.").ToString();
+			*OutMsg = FString::Printf(*NSLOCTEXT("K2Node", "Function_not_compatible", "The function/event '%s' does not match the necessary signature - has the delegate or function/event changed?").ToString(),
+				*FunctionName.ToString());
 		}
 		return false;
 	}
@@ -115,7 +134,8 @@ bool UK2Node_CreateDelegate::IsValid(FString* OutMsg, bool bDontUseSkeletalClass
 	{
 		if (OutMsg)
 		{
-			*OutMsg = NSLOCTEXT("K2Node", "Function_cannot_be_used_in_delegate", "Function cannot be used in delegate.").ToString();
+			*OutMsg = FString::Printf(*NSLOCTEXT("K2Node", "Function_cannot_be_used_in_delegate", "The selected function/event is not bindable - is the function/event pure or latent?").ToString(),
+				*FunctionName.ToString());
 		}
 		return false;
 	}
@@ -130,7 +150,8 @@ bool UK2Node_CreateDelegate::IsValid(FString* OutMsg, bool bDontUseSkeletalClass
 			{
 				if(OutMsg)
 				{
-					*OutMsg = NSLOCTEXT("K2Node", "WrongDelegateAuthorityOnly", "No AuthorityOnly flag").ToString();
+					*OutMsg = FString::Printf(*NSLOCTEXT("K2Node", "WrongDelegateAuthorityOnly", "The selected function/event ('%s') is not compatible with this delegate (the delegate is server-only) - try marking the function/event AuthorityOnly.").ToString(),
+						*FunctionName.ToString());
 				}
 				return false;
 			}
@@ -145,7 +166,7 @@ void UK2Node_CreateDelegate::ValidationAfterFunctionsAreCreated(class FCompilerR
 	FString Msg;
 	if(!IsValid(&Msg, bFullCompile))
 	{
-		MessageLog.Error(*FString::Printf( TEXT("%s %s @@"), *NSLOCTEXT("K2Node", "WrongDelegate", "Event signature error: ").ToString(), *Msg), this);
+ 		MessageLog.Error(*FString::Printf( TEXT("@@ %s %s"), *NSLOCTEXT("K2Node", "WrongDelegate", "Signature Error:").ToString(), *Msg), this);
 	}
 }
 
@@ -172,7 +193,16 @@ void UK2Node_CreateDelegate::HandleAnyChangeWithoutNotifying()
 
 	if(!IsValid())
 	{
-		SelectedFunctionName = NAME_None;
+		// do not clear the name, so we can keep it around as a hint/guide for 
+		// users (so they can better determine what went wrong)
+		if (const UEdGraphPin* DelegatePin = GetDelegateOutPin())
+		{
+			if (DelegatePin->LinkedTo.Num() == 0)
+			{
+				// ok to clear if they've disconnected the delegate pin
+				SelectedFunctionName = NAME_None;
+			}
+		}
 		SelectedFunctionGuid.Invalidate();
 	}
 }
@@ -340,8 +370,7 @@ UEdGraphPin* UK2Node_CreateDelegate::GetDelegateOutPin() const
 
 UEdGraphPin* UK2Node_CreateDelegate::GetObjectInPin() const
 {
-	const UEdGraphSchema_K2* K2Schema = GetDefault<UEdGraphSchema_K2>();
-	return FindPinChecked(K2Schema->PN_Self);
+	return FindPin(UEdGraphSchema_K2::PN_Self);
 }
 
 FText UK2Node_CreateDelegate::GetNodeTitle(ENodeTitleType::Type TitleType) const
diff --git a/Engine/Source/Editor/BlueprintGraph/Private/K2Node_GetClassDefaults.cpp b/Engine/Source/Editor/BlueprintGraph/Private/K2Node_GetClassDefaults.cpp
--- a/Engine/Source/Editor/BlueprintGraph/Private/K2Node_GetClassDefaults.cpp
+++ b/Engine/Source/Editor/BlueprintGraph/Private/K2Node_GetClassDefaults.cpp
@@ -349,29 +349,32 @@ UClass* UK2Node_GetClassDefaults::GetInputClass(const UEdGraphPin* FromPin) cons
 {
 	UClass* InputClass = nullptr;
 
-	const UEdGraphPin* ClassPin = FromPin;
-	if(ClassPin == nullptr)
+	if (FromPin != nullptr)
 	{
-		ClassPin = FindClassPin();
-	}
-	
-	if(ClassPin != nullptr)
-	{
-		check(ClassPin->Direction == EGPD_Input);
+		check(FromPin->Direction == EGPD_Input);
 
-		if(ClassPin->DefaultObject != nullptr && ClassPin->LinkedTo.Num() == 0)
+		if (FromPin->DefaultObject != nullptr && FromPin->LinkedTo.Num() == 0)
 		{
-			InputClass = CastChecked<UClass>(ClassPin->DefaultObject);
+			InputClass = CastChecked<UClass>(FromPin->DefaultObject);
 		}
-		else if(ClassPin->LinkedTo.Num() > 0)
+		else if (FromPin->LinkedTo.Num() > 0)
 		{
-			if(UEdGraphPin* LinkedPin = ClassPin->LinkedTo[0])
+			if (UEdGraphPin* LinkedPin = FromPin->LinkedTo[0])
 			{
 				InputClass = Cast<UClass>(LinkedPin->PinType.PinSubCategoryObject.Get());
 			}
 		}
 	}
 
+	// Switch Blueprint Class types to use the generated skeleton class.
+	if (InputClass)
+	{
+		if (UBlueprint* Blueprint = Cast<UBlueprint>(InputClass->ClassGeneratedBy))
+		{
+			InputClass = Blueprint->SkeletonGeneratedClass;
+		}
+	}
+
 	return InputClass;
 }
 
@@ -401,6 +404,25 @@ void UK2Node_GetClassDefaults::CreateOutputPins(UClass* InClass)
 	{
 		AdvancedPinDisplay = ENodeAdvancedPins::NoPins;
 	}
+
+	// Unbind OnChanged() delegate from a previous Blueprint, if valid.
+	if (OnBlueprintChangedDelegate.IsValid())
+	{
+		OnBlueprintChangedDelegate.Reset();
+	}
+
+	// Unbind OnCompiled() delegate from a previous Blueprint, if valid.
+	if (OnBlueprintCompiledDelegate.IsValid())
+	{
+		OnBlueprintCompiledDelegate.Reset();
+	}
+
+	// If the class was generated for a Blueprint, bind delegates to handle any OnChanged() & OnCompiled() events.
+	if (UBlueprint* Blueprint = Cast<UBlueprint>(InClass->ClassGeneratedBy))
+	{
+		OnBlueprintChangedDelegate = Blueprint->OnChanged().AddLambda([this](UBlueprint* /* InBlueprint */) { ReconstructNode(); });
+		OnBlueprintCompiledDelegate = Blueprint->OnCompiled().AddLambda([this](UBlueprint* /* InBlueprint */) { ReconstructNode(); });
+	}
 }
 
 void UK2Node_GetClassDefaults::OnClassPinChanged()
diff --git a/Engine/Source/Editor/BlueprintGraph/Private/K2Node_MakeStruct.cpp b/Engine/Source/Editor/BlueprintGraph/Private/K2Node_MakeStruct.cpp
--- a/Engine/Source/Editor/BlueprintGraph/Private/K2Node_MakeStruct.cpp
+++ b/Engine/Source/Editor/BlueprintGraph/Private/K2Node_MakeStruct.cpp
@@ -361,67 +361,70 @@ void UK2Node_MakeStruct::Serialize(FArchive& Ar)
 		{
 			// Check if this node actually requires warning the user that functionality has changed.
 			bMadeAfterOverridePinRemoval = true;
-			FOptionalPinManager PinManager;
-
-			// Preload all pins so we can examine their values
-			for (UEdGraphPin* Pin : Pins)
+			if (StructType != nullptr)
 			{
-				Ar.Preload(Pin);
-			}
+				FOptionalPinManager PinManager;
 
-			// Have to check if this node is even in danger.
-			for (FOptionalPinFromProperty& PropertyEntry : ShowPinForProperties)
-			{
-				UProperty* Property = StructType->FindPropertyByName(PropertyEntry.PropertyName);
-				bool bNegate = false;
-				if (UProperty* OverrideProperty = PropertyCustomizationHelpers::GetEditConditionProperty(Property, bNegate))
+				// Preload all pins so we can examine their values
+				for (UEdGraphPin* Pin : Pins)
 				{
-					bool bHadOverridePropertySeparation = false;
-					for (FOptionalPinFromProperty& OverridePropertyEntry : ShowPinForProperties)
+					Ar.Preload(Pin);
+				}
+
+				// Have to check if this node is even in danger.
+				for (FOptionalPinFromProperty& PropertyEntry : ShowPinForProperties)
+				{
+					UProperty* Property = StructType->FindPropertyByName(PropertyEntry.PropertyName);
+					bool bNegate = false;
+					if (UProperty* OverrideProperty = PropertyCustomizationHelpers::GetEditConditionProperty(Property, bNegate))
 					{
-						if (OverridePropertyEntry.PropertyName == OverrideProperty->GetFName())
+						bool bHadOverridePropertySeparation = false;
+						for (FOptionalPinFromProperty& OverridePropertyEntry : ShowPinForProperties)
 						{
-							bHadOverridePropertySeparation = true;
-							break;
+							if (OverridePropertyEntry.PropertyName == OverrideProperty->GetFName())
+							{
+								bHadOverridePropertySeparation = true;
+								break;
+							}
 						}
-					}
 
-					bMadeAfterOverridePinRemoval = false;
-					UEdGraphPin* Pin = FindPin(Property->GetName());
+						bMadeAfterOverridePinRemoval = false;
+						UEdGraphPin* Pin = FindPin(Property->GetName());
 
-					if (bHadOverridePropertySeparation)
-					{
-						UEdGraphPin* OverridePin = FindPin(OverrideProperty->GetName());
-						const UEdGraphSchema_K2* Schema = GetDefault<UEdGraphSchema_K2>();
-						if (OverridePin)
+						if (bHadOverridePropertySeparation)
 						{
-							// Override pins are always booleans
-							check(OverridePin->PinType.PinCategory == Schema->PC_Boolean);
-							// If the old override pin's default value was true, then the override should be marked as enabled
-							PropertyEntry.bIsOverrideEnabled = OverridePin->DefaultValue.ToBool();
-							// It had an override pin, so conceptually the override pin is visible
-							PropertyEntry.bIsOverridePinVisible = true;
-							// Because there was an override pin visible for this property, this property will be forced to have a pin
-							PropertyEntry.bShowPin = true;
+							UEdGraphPin* OverridePin = FindPin(OverrideProperty->GetName());
+							const UEdGraphSchema_K2* Schema = GetDefault<UEdGraphSchema_K2>();
+							if (OverridePin)
+							{
+								// Override pins are always booleans
+								check(OverridePin->PinType.PinCategory == Schema->PC_Boolean);
+								// If the old override pin's default value was true, then the override should be marked as enabled
+								PropertyEntry.bIsOverrideEnabled = OverridePin->DefaultValue.ToBool();
+								// It had an override pin, so conceptually the override pin is visible
+								PropertyEntry.bIsOverridePinVisible = true;
+								// Because there was an override pin visible for this property, this property will be forced to have a pin
+								PropertyEntry.bShowPin = true;
+							}
+							else
+							{
+								// No override pin, ensure all override bools are false
+								PropertyEntry.bIsOverrideEnabled = false;
+								PropertyEntry.bIsOverridePinVisible = false;
+							}
 						}
 						else
 						{
-							// No override pin, ensure all override bools are false
-							PropertyEntry.bIsOverrideEnabled = false;
-							PropertyEntry.bIsOverridePinVisible = false;
+							if (Pin)
+							{
+								PropertyEntry.bIsOverrideEnabled = true;
+								PropertyEntry.bIsOverridePinVisible = true;
+							}
 						}
-					}
-					else
-					{
-						if (Pin)
-						{
-							PropertyEntry.bIsOverrideEnabled = true;
-							PropertyEntry.bIsOverridePinVisible = true;
-						}
-					}
 
-					// If the pin for this property, which sets the property's value, does not exist then the user was not trying to set the value
-					PropertyEntry.bIsSetValuePinVisible = Pin != nullptr;
+						// If the pin for this property, which sets the property's value, does not exist then the user was not trying to set the value
+						PropertyEntry.bIsSetValuePinVisible = Pin != nullptr;
+					}
 				}
 			}
 		}
diff --git a/Engine/Source/Editor/BlueprintGraph/Private/K2Node_SpawnActorFromClass.cpp b/Engine/Source/Editor/BlueprintGraph/Private/K2Node_SpawnActorFromClass.cpp
--- a/Engine/Source/Editor/BlueprintGraph/Private/K2Node_SpawnActorFromClass.cpp
+++ b/Engine/Source/Editor/BlueprintGraph/Private/K2Node_SpawnActorFromClass.cpp
@@ -375,16 +375,26 @@ void UK2Node_SpawnActorFromClass::GetPinHoverText(const UEdGraphPin& Pin, FStrin
 {
 	const UEdGraphSchema_K2* K2Schema = GetDefault<UEdGraphSchema_K2>();
 
-	UEdGraphPin* ClassPin = GetClassPin();
-	K2Schema->ConstructBasicPinTooltip(*ClassPin, LOCTEXT("ClassPinDescription", "The Actor class you want to spawn"), ClassPin->PinToolTip);
-	UEdGraphPin* TransformPin = GetSpawnTransformPin();
-	K2Schema->ConstructBasicPinTooltip(*TransformPin, LOCTEXT("TransformPinDescription", "The transform to spawn the Actor with"), TransformPin->PinToolTip);
-	UEdGraphPin* CollisionHandlingOverridePin = GetCollisionHandlingOverridePin();
-	K2Schema->ConstructBasicPinTooltip(*CollisionHandlingOverridePin, LOCTEXT("CollisionHandlingOverridePinDescription", "Specifies how to handle collisions at the spawn point. If undefined, uses actor class settings."), CollisionHandlingOverridePin->PinToolTip);
-	UEdGraphPin* OwnerPin = GetOwnerPin();
-	K2Schema->ConstructBasicPinTooltip(*OwnerPin, LOCTEXT("OwnerPinDescription", "Can be left empty; primarily used for replication (bNetUseOwnerRelevancy and bOnlyRelevantToOwner), or visibility (PrimitiveComponent's bOwnerNoSee/bOnlyOwnerSee)"), OwnerPin->PinToolTip);
-	UEdGraphPin* ResultPin = GetResultPin();
-	K2Schema->ConstructBasicPinTooltip(*ResultPin, LOCTEXT("ResultPinDescription", "The spawned Actor"), ResultPin->PinToolTip);
+	if (UEdGraphPin* ClassPin = GetClassPin())
+	{
+		K2Schema->ConstructBasicPinTooltip(*ClassPin, LOCTEXT("ClassPinDescription", "The Actor class you want to spawn"), ClassPin->PinToolTip);
+	}
+	if (UEdGraphPin* TransformPin = GetSpawnTransformPin())
+	{
+		K2Schema->ConstructBasicPinTooltip(*TransformPin, LOCTEXT("TransformPinDescription", "The transform to spawn the Actor with"), TransformPin->PinToolTip);
+	}
+	if (UEdGraphPin* CollisionHandlingOverridePin = GetCollisionHandlingOverridePin())
+	{
+		K2Schema->ConstructBasicPinTooltip(*CollisionHandlingOverridePin, LOCTEXT("CollisionHandlingOverridePinDescription", "Specifies how to handle collisions at the spawn point. If undefined, uses actor class settings."), CollisionHandlingOverridePin->PinToolTip);
+	}
+	if (UEdGraphPin* OwnerPin = GetOwnerPin())
+	{
+		K2Schema->ConstructBasicPinTooltip(*OwnerPin, LOCTEXT("OwnerPinDescription", "Can be left empty; primarily used for replication (bNetUseOwnerRelevancy and bOnlyRelevantToOwner), or visibility (PrimitiveComponent's bOwnerNoSee/bOnlyOwnerSee)"), OwnerPin->PinToolTip);
+	}
+	if (UEdGraphPin* ResultPin = GetResultPin())
+	{
+		K2Schema->ConstructBasicPinTooltip(*ResultPin, LOCTEXT("ResultPinDescription", "The spawned Actor"), ResultPin->PinToolTip);
+	}
 
 	return Super::GetPinHoverText(Pin, HoverTextOut);
 }
diff --git a/Engine/Source/Editor/BlueprintGraph/Public/BlueprintGraphDefinitions.h b/Engine/Source/Editor/BlueprintGraph/Public/BlueprintGraphDefinitions.h
--- a/Engine/Source/Editor/BlueprintGraph/Public/BlueprintGraphDefinitions.h
+++ b/Engine/Source/Editor/BlueprintGraph/Public/BlueprintGraphDefinitions.h
@@ -34,6 +34,7 @@
 #include "K2Node_FormatText.h"
 #include "K2Node_FunctionEntry.h"
 #include "K2Node_FunctionResult.h"
+#include "K2Node_GetClassDefaults.h"
 #include "K2Node_GetDataTableRow.h"
 #include "K2Node_GetArrayItem.h"
 #include "K2Node_IfThenElse.h"
diff --git a/Engine/Source/Editor/ClassViewer/Private/SClassViewer.cpp b/Engine/Source/Editor/ClassViewer/Private/SClassViewer.cpp
--- a/Engine/Source/Editor/ClassViewer/Private/SClassViewer.cpp
+++ b/Engine/Source/Editor/ClassViewer/Private/SClassViewer.cpp
@@ -33,6 +33,7 @@
 #include "SourceCodeNavigation.h"
 #include "HotReloadInterface.h"
 #include "SSearchBox.h"
+#include "TextFilterExpressionEvaluator.h"
 
 #include "SListViewSelectorDropdownMenu.h"
 #include "Engine/BlueprintGeneratedClass.h"
@@ -535,27 +536,36 @@ namespace ClassViewer
 		}
 
 		/** Checks if the TestString passes the filter.
-		 *	@param InTestString				The string to test against the filter.
-		 *	@param InFilterSearchTerms		A list of terms to filter with.
+		 *	@param InTestString			The string to test against the filter.
+		 *	@param InTextFilter			Compiled text filter to apply.
 		 *
 		 *	@return	true if it passes the filter.
 		 */
-		static bool PassesFilter( const FString& InTestString, const TArray<FString>& InFilterSearchTerms )
+		static bool PassesFilter( const FString& InTestString, const FTextFilterExpressionEvaluator& InTextFilter )
 		{
-			bool bPassesFilter = false;
+			class FClassFilterContext : public ITextFilterExpressionContext
+			{
+			public:
+				explicit FClassFilterContext(const FString& InStr)
+					: StrPtr(&InStr)
+				{
+				}
 
-			bPassesFilter = !InFilterSearchTerms.Num();
+				virtual bool TestBasicStringExpression(const FTextFilterString& InValue, const ETextFilterTextComparisonMode InTextComparisonMode) const override
+				{
+					return TextFilterUtils::TestBasicStringExpression(*StrPtr, InValue, InTextComparisonMode);
+				}
 
-			for( auto CurSearchTermIndex = 0; !bPassesFilter && CurSearchTermIndex < InFilterSearchTerms.Num(); ++CurSearchTermIndex )
-			{
-				if(InTestString.Contains(InFilterSearchTerms[CurSearchTermIndex]))
+				virtual bool TestComplexExpression(const FName& InKey, const FTextFilterString& InValue, const ETextFilterComparisonOperation InComparisonOperation, const ETextFilterTextComparisonMode InTextComparisonMode) const override
 				{
-					bPassesFilter = true;
-					break;
+					return false;
 				}
-			}
 
-			return bPassesFilter;
+			private:
+				const FString* StrPtr;
+			};
+
+			return InTextFilter.TestTextFilter(FClassFilterContext(InTestString));
 		}
 
 		/** Will create the instance of FClassHierarchy and populate the class hierarchy tree. */
@@ -607,7 +617,7 @@ namespace ClassViewer
 		 *	@param InInitOptions						The class viewer's options, holds the AllowedClasses and DisallowedClasses.
 		 *	@param InOutRootNode						The node that this function will add the children of to the tree.
 		 *	@param InRootClassIndex						The index of the root node.
-		 *	@param InFilterSearchTerms					A list of terms to filter with.
+		 *	@param InTextFilter							Compiled text filter to apply.
 		 *	@param bInOnlyActors						Filter option to remove non-actor classes.
 		 *	@param bInOnlyPlaceables					Filter option to remove non-placeable classes.
 		 *	@param bInOnlyBlueprintBases				Filter option to remove non-blueprint base classes.
@@ -615,7 +625,7 @@ namespace ClassViewer
 		 *
 		 *	@return Returns true if the child passed the filter.
 		 */		
-		static bool AddChildren_Tree(const FClassViewerInitializationOptions& InInitOptions, TSharedPtr< FClassViewerNode >& InOutRootNode, const TSharedPtr< FClassViewerNode >& InOriginalRootNode, const TArray<FString>& InFilterSearchTerms, bool bInOnlyActors, bool bInOnlyPlaceables, bool bInOnlyBlueprintBases, bool bInShowUnloadedBlueprints )
+		static bool AddChildren_Tree(const FClassViewerInitializationOptions& InInitOptions, TSharedPtr< FClassViewerNode >& InOutRootNode, const TSharedPtr< FClassViewerNode >& InOriginalRootNode, const FTextFilterExpressionEvaluator& InTextFilter, bool bInOnlyActors, bool bInOnlyPlaceables, bool bInOnlyBlueprintBases, bool bInShowUnloadedBlueprints )
 		{
 			if (bInOnlyActors && *InOriginalRootNode->GetClassName().Get() != FString(TEXT("Actor")))
 			{
@@ -645,12 +655,12 @@ namespace ClassViewer
 			{
 				if(bInShowUnloadedBlueprints)
 				{
-					bReturnPassesFilter = InOutRootNode->bPassesFilter = bPassesPlaceableFilter && bPassesBlueprintBaseFilter && IsClassAllowed_UnloadedBlueprint(InInitOptions, InOriginalRootNode) && PassesFilter(*InOriginalRootNode->GetClassName().Get(), InFilterSearchTerms);
+					bReturnPassesFilter = InOutRootNode->bPassesFilter = bPassesPlaceableFilter && bPassesBlueprintBaseFilter && IsClassAllowed_UnloadedBlueprint(InInitOptions, InOriginalRootNode) && PassesFilter(*InOriginalRootNode->GetClassName().Get(), InTextFilter);
 				}
 			}
 			else
 			{
-				bReturnPassesFilter = InOutRootNode->bPassesFilter = bPassesPlaceableFilter && bPassesBlueprintBaseFilter && IsClassAllowed(InInitOptions, InOriginalRootNode->Class) && PassesFilter(*InOriginalRootNode->GetClassName().Get(), InFilterSearchTerms);
+				bReturnPassesFilter = InOutRootNode->bPassesFilter = bPassesPlaceableFilter && bPassesBlueprintBaseFilter && IsClassAllowed(InInitOptions, InOriginalRootNode->Class) && PassesFilter(*InOriginalRootNode->GetClassName().Get(), InTextFilter);
 			}
 
 			TArray< TSharedPtr< FClassViewerNode > >& ChildList = InOriginalRootNode->GetChildrenList();
@@ -658,7 +668,7 @@ namespace ClassViewer
 			{
 				TSharedPtr< FClassViewerNode > NewNode = MakeShareable( new FClassViewerNode( *ChildList[ChildIdx].Get() ) );
 
-				bReturnPassesFilter |= bChildrenPassesFilter = AddChildren_Tree(InInitOptions, NewNode, ChildList[ChildIdx], InFilterSearchTerms, false, /* bInOnlyActors - false so that anything below Actor is added */
+				bReturnPassesFilter |= bChildrenPassesFilter = AddChildren_Tree(InInitOptions, NewNode, ChildList[ChildIdx], InTextFilter, false, /* bInOnlyActors - false so that anything below Actor is added */
 					bInOnlyPlaceables, bInOnlyBlueprintBases, bInShowUnloadedBlueprints);
 				if(bChildrenPassesFilter)
 				{
@@ -673,7 +683,7 @@ namespace ClassViewer
 		/** Builds the class tree.
 		 *	@param InInitOptions						The class viewer's options, holds the AllowedClasses and DisallowedClasses.
 		 *	@param InOutRootNode						The node to root the tree to.
-		 *	@param InFilterSearchTerms					A list of terms to filter with.
+		 *	@param InTextFilter							Compiled text filter to apply.
 		 *	@param bInOnlyPlaceables					Filter option to remove non-placeable classes.
 		 *	@param bInOnlyActors						Filter option to root the tree in the "Actor" class.
 		 *	@param bInOnlyBlueprintBases				Filter option to remove non-blueprint base classes.
@@ -681,7 +691,7 @@ namespace ClassViewer
 		 *
 		 *	@return A fully built tree.
 		 */
-		static void GetClassTree(const FClassViewerInitializationOptions& InInitOptions, TSharedPtr< FClassViewerNode >& InOutRootNode, const TArray<FString>& InFilterSearchTerms, bool bInOnlyPlaceables, bool bInOnlyActors, bool bInOnlyBlueprintBases, bool bInShowUnloadedBlueprints )
+		static void GetClassTree(const FClassViewerInitializationOptions& InInitOptions, TSharedPtr< FClassViewerNode >& InOutRootNode, const FTextFilterExpressionEvaluator& InTextFilter, bool bInOnlyPlaceables, bool bInOnlyActors, bool bInOnlyBlueprintBases, bool bInShowUnloadedBlueprints )
 		{
 			const TSharedPtr< FClassViewerNode > ObjectClassRoot = ClassHierarchy->GetObjectRootNode();
 
@@ -693,13 +703,13 @@ namespace ClassViewer
 				for(int32 ClassIdx = 0; ClassIdx < ObjectClassRoot->GetChildrenList().Num(); ClassIdx++)
 				{
 					TSharedPtr<FClassViewerNode> ChildNode = MakeShareable(new FClassViewerNode(*ObjectClassRoot->GetChildrenList()[ClassIdx].Get()));
-					if (AddChildren_Tree(InInitOptions, ChildNode, ObjectClassRoot->GetChildrenList()[ClassIdx], InFilterSearchTerms, true, bInOnlyPlaceables, bInOnlyBlueprintBases, bInShowUnloadedBlueprints))
+					if (AddChildren_Tree(InInitOptions, ChildNode, ObjectClassRoot->GetChildrenList()[ClassIdx], InTextFilter, true, bInOnlyPlaceables, bInOnlyBlueprintBases, bInShowUnloadedBlueprints))
 						InOutRootNode->AddChild(ChildNode);
 				}
 			}
 			else
 			{
-				AddChildren_Tree(InInitOptions, InOutRootNode, ObjectClassRoot, InFilterSearchTerms, false, bInOnlyPlaceables, bInOnlyBlueprintBases, bInShowUnloadedBlueprints);				
+				AddChildren_Tree(InInitOptions, InOutRootNode, ObjectClassRoot, InTextFilter, false, bInOnlyPlaceables, bInOnlyBlueprintBases, bInShowUnloadedBlueprints);				
 			}
 		}
 
@@ -707,7 +717,7 @@ namespace ClassViewer
 		 *	@param InInitOptions						The class viewer's options, holds the AllowedClasses and DisallowedClasses.
 		 *	@param InOutRootNode						The node that this function will add the children of to the tree.
 		 *	@param InRootClassIndex						The index of the root node.
-		 *	@param InFilterSearchTerms					A list of terms to filter with.
+		 *	@param InTextFilter							Compiled text filter to apply.
 		 *	@param bInOnlyActors						Filter option to remove non-actor classes.
 		 *	@param bInOnlyPlaceables					Filter option to remove non-placeable classes.
 		 *	@param bInOnlyBlueprintBases				Filter option to remove non-blueprint base classes.
@@ -715,7 +725,7 @@ namespace ClassViewer
 		 *
 		 *	@return Returns true if the child passed the filter.
 		 */
-		static void AddChildren_List( const FClassViewerInitializationOptions& InInitOptions, TArray< TSharedPtr< FClassViewerNode > >& InOutNodeList, const TSharedPtr< FClassViewerNode >& InOriginalRootNode, const TArray<FString>& InFilterSearchTerms, bool bInOnlyActors, bool bInOnlyPlaceables, bool bInOnlyBlueprintBases, bool bInShowUnloadedBlueprints )
+		static void AddChildren_List( const FClassViewerInitializationOptions& InInitOptions, TArray< TSharedPtr< FClassViewerNode > >& InOutNodeList, const TSharedPtr< FClassViewerNode >& InOriginalRootNode, const FTextFilterExpressionEvaluator& InTextFilter, bool bInOnlyActors, bool bInOnlyPlaceables, bool bInOnlyBlueprintBases, bool bInShowUnloadedBlueprints )
 		{
 			if (bInOnlyActors && *InOriginalRootNode->GetClassName().Get() != FString(TEXT("Actor")))
 			{
@@ -743,12 +753,12 @@ namespace ClassViewer
 			{
 				if(bInShowUnloadedBlueprints)
 				{
-					NewNode->bPassesFilter = bPassesPlaceableFilter && bPassesBlueprintBaseFilter && IsClassAllowed_UnloadedBlueprint(InInitOptions, InOriginalRootNode) && PassesFilter(*InOriginalRootNode->GetClassName().Get(), InFilterSearchTerms);
+					NewNode->bPassesFilter = bPassesPlaceableFilter && bPassesBlueprintBaseFilter && IsClassAllowed_UnloadedBlueprint(InInitOptions, InOriginalRootNode) && PassesFilter(*InOriginalRootNode->GetClassName().Get(), InTextFilter);
 				}
 			}
 			else
 			{
-				NewNode->bPassesFilter = bPassesPlaceableFilter && bPassesBlueprintBaseFilter && IsClassAllowed(InInitOptions, InOriginalRootNode->Class) && PassesFilter(*InOriginalRootNode->GetClassName().Get(), InFilterSearchTerms);
+				NewNode->bPassesFilter = bPassesPlaceableFilter && bPassesBlueprintBaseFilter && IsClassAllowed(InInitOptions, InOriginalRootNode->Class) && PassesFilter(*InOriginalRootNode->GetClassName().Get(), InTextFilter);
 			}
 
 			if(NewNode->bPassesFilter)
@@ -761,7 +771,7 @@ namespace ClassViewer
 			TArray< TSharedPtr< FClassViewerNode > >& ChildList = InOriginalRootNode->GetChildrenList();
 			for(int32 ChildIdx = 0; ChildIdx < ChildList.Num(); ChildIdx++)
 			{
-				AddChildren_List(InInitOptions, InOutNodeList, ChildList[ChildIdx], InFilterSearchTerms, false, /* bInOnlyActors - false so that anything below Actor is added */
+				AddChildren_List(InInitOptions, InOutNodeList, ChildList[ChildIdx], InTextFilter, false, /* bInOnlyActors - false so that anything below Actor is added */
 					bInOnlyPlaceables, bInOnlyBlueprintBases, bInShowUnloadedBlueprints);
 			}
 		}	
@@ -769,7 +779,7 @@ namespace ClassViewer
 		/** Builds the class list.
 		 *	@param InInitOptions						The class viewer's options, holds the AllowedClasses and DisallowedClasses.
 		 *	@param InOutNodeList						The list to add all the nodes to.
-		 *	@param InFilterSearchTerms					A list of terms to filter with.
+		 *	@param InTextFilter							Compiled text filter to apply.
 		 *	@param bInOnlyPlaceables					Filter option to remove non-placeable classes.
 		 *	@param bInOnlyActors						Filter option to root the tree in the "Actor" class.
 		 *	@param bInOnlyBlueprintBases				Filter option to remove non-blueprint base classes.
@@ -777,7 +787,7 @@ namespace ClassViewer
 		 *
 		 *	@return A fully built list.
 		 */
-		static void GetClassList(const FClassViewerInitializationOptions& InInitOptions, TArray< TSharedPtr< FClassViewerNode > >& InOutNodeList, const TArray<FString>& InFilterSearchTerms, bool bInOnlyPlaceables, bool bInOnlyActors, bool bInOnlyBlueprintBases, bool bInShowUnloadedBlueprints)
+		static void GetClassList(const FClassViewerInitializationOptions& InInitOptions, TArray< TSharedPtr< FClassViewerNode > >& InOutNodeList, const FTextFilterExpressionEvaluator& InTextFilter, bool bInOnlyPlaceables, bool bInOnlyActors, bool bInOnlyBlueprintBases, bool bInShowUnloadedBlueprints)
 		{
 			const TSharedPtr< FClassViewerNode > ObjectClassRoot = ClassHierarchy->GetObjectRootNode();
 
@@ -785,7 +795,7 @@ namespace ClassViewer
 			if(InInitOptions.bShowObjectRootClass)
 			{
 				TSharedPtr< FClassViewerNode > NewNode = MakeShareable( new FClassViewerNode( *ObjectClassRoot.Get() ) );
-				NewNode->bPassesFilter = IsClassAllowed(InInitOptions, ObjectClassRoot->Class) && PassesFilter(*ObjectClassRoot->GetClassName().Get(), InFilterSearchTerms);
+				NewNode->bPassesFilter = IsClassAllowed(InInitOptions, ObjectClassRoot->Class) && PassesFilter(*ObjectClassRoot->GetClassName().Get(), InTextFilter);
 				
 				if(NewNode->bPassesFilter)
 				{
@@ -798,7 +808,7 @@ namespace ClassViewer
 			TArray< TSharedPtr< FClassViewerNode > >& ChildList = ObjectClassRoot->GetChildrenList();
 			for(int32 ObjectChildIndex = 0; ObjectChildIndex < ChildList.Num(); ObjectChildIndex++)
 			{
-				AddChildren_List(InInitOptions, InOutNodeList, ChildList[ObjectChildIndex], InFilterSearchTerms, bInOnlyActors, bInOnlyPlaceables, bInOnlyBlueprintBases, bInShowUnloadedBlueprints);			
+				AddChildren_List(InInitOptions, InOutNodeList, ChildList[ObjectChildIndex], InTextFilter, bInOnlyActors, bInOnlyPlaceables, bInOnlyBlueprintBases, bInShowUnloadedBlueprints);			
 			}
 		}
 
@@ -1924,6 +1934,8 @@ void SClassViewer::Construct(const FArguments& InArgs, const FClassViewerInitial
 
 	OnClassPicked = InArgs._OnClassPickedDelegate;
 
+	TextFilterPtr = MakeShareable(new FTextFilterExpressionEvaluator(ETextFilterExpressionEvaluatorMode::BasicString));
+
 	bSaveExpansionStates = true;
 	bPendingSetExpansionStates = false;
 
@@ -2367,56 +2379,9 @@ void SClassViewer::FindInContentBrowser()
 
 void SClassViewer::OnFilterTextChanged( const FText& InFilterText )
 {
-	// If the filter was empty, we want it to save the expansion states so when it returns to empty it can restore them.
-	bSaveExpansionStates = !FilterSearchTerms.Num();
-
-	FilterSearchTerms.Empty();
-
-	FString CurrentFilterText = InFilterText.ToString();
-
-	// Sanitize, parse and split the filter text into search terms
-	CurrentFilterText.Trim().TrimTrailing();
-	//FilterSearchTerms.Reset();
-	{
-		bool bIsWithinQuotedSection = false;
-		FString NewSearchTerm;
-		for( auto CurChar : CurrentFilterText )
-		{
-			// Keep an eye out for double-quotes.  We want to retain whitespace within a search term if
-			// it has double-quotes around it
-			if( CurChar == TCHAR('\"') )
-			{
-				// Toggle whether we're within a quoted section, but don't bother adding the quote character
-				bIsWithinQuotedSection = !bIsWithinQuotedSection;
-			}
-			else if( bIsWithinQuotedSection || !FChar::IsWhitespace( CurChar ) )
-			{
-				// Add the character!
-				NewSearchTerm.AppendChar( CurChar );
-			}
-			else
-			{
-				// Encountered whitespace, so add the search term up to here
-				if( NewSearchTerm.Len() > 0 )
-				{
-					FilterSearchTerms.Add( NewSearchTerm );
-					NewSearchTerm = FString();
-				}
-			}
-		}
-
-		// Encountered EOL, so add the search term up to here
-		if( NewSearchTerm.Len() > 0 )
-		{
-			FilterSearchTerms.Add( NewSearchTerm );
-			NewSearchTerm = FString();
-		}
-
-		if( bIsWithinQuotedSection )
-		{
-			// User forgot to terminate their double-quoted section.  No big deal.
-		}
-	}
+	// Update the compiled filter and report any syntax error information back to the user
+	TextFilterPtr->SetFilterText(InFilterText);
+	SearchBox->SetError(TextFilterPtr->GetFilterErrorText());
 
 	// Repopulate the list to show only what has not been filtered out.
 	Refresh();
@@ -2644,10 +2609,10 @@ void SClassViewer::Populate()
 		TSharedPtr<FClassViewerNode> RootNode;
 
 		// Get the class tree, passing in certain filter options.
-		ClassViewer::Helpers::GetClassTree(InitOptions, RootNode, FilterSearchTerms, MenuPlaceableOnly_IsChecked(), MenuActorsOnly_IsChecked(), MenuBlueprintBasesOnly_IsChecked(), bShowUnloadedBlueprints );
+		ClassViewer::Helpers::GetClassTree(InitOptions, RootNode, *TextFilterPtr, MenuPlaceableOnly_IsChecked(), MenuActorsOnly_IsChecked(), MenuBlueprintBasesOnly_IsChecked(), bShowUnloadedBlueprints );
 
 		// Check if we will restore expansion states, we will not if there is filtering happening.
-		const bool bRestoreExpansionState = !FilterSearchTerms.Num();
+		const bool bRestoreExpansionState = TextFilterPtr->GetFilterType() == ETextFilterExpressionType::Empty;
 
 		if(InitOptions.bShowObjectRootClass)
 		{
@@ -2659,7 +2624,7 @@ void SClassViewer::Populate()
 			}
 
 			// Expand any items that pass the filter.
-			if(FilterSearchTerms.Num() > 0)
+			if(TextFilterPtr->GetFilterType() != ETextFilterExpressionType::Empty)
 			{
 				ExpandFilteredInNodes(RootNode);
 			}
@@ -2676,7 +2641,7 @@ void SClassViewer::Populate()
 				}
 
 				// Expand any items that pass the filter.
-				if(FilterSearchTerms.Num() > 0)
+				if(TextFilterPtr->GetFilterType() != ETextFilterExpressionType::Empty)
 				{
 					ExpandFilteredInNodes(RootNode->GetChildrenList()[ChildIndex]);
 				}
@@ -2697,7 +2662,7 @@ void SClassViewer::Populate()
 	else
 	{
 		// Get the class list, passing in certain filter options.
-		ClassViewer::Helpers::GetClassList(InitOptions, RootTreeItems, FilterSearchTerms, MenuPlaceableOnly_IsChecked(), MenuActorsOnly_IsChecked(), MenuBlueprintBasesOnly_IsChecked(), bShowUnloadedBlueprints);
+		ClassViewer::Helpers::GetClassList(InitOptions, RootTreeItems, *TextFilterPtr, MenuPlaceableOnly_IsChecked(), MenuActorsOnly_IsChecked(), MenuBlueprintBasesOnly_IsChecked(), bShowUnloadedBlueprints);
 
 		// Sort the list alphabetically.
 		struct FCompareFClassViewerNode
diff --git a/Engine/Source/Editor/ClassViewer/Private/SClassViewer.h b/Engine/Source/Editor/ClassViewer/Private/SClassViewer.h
--- a/Engine/Source/Editor/ClassViewer/Private/SClassViewer.h
+++ b/Engine/Source/Editor/ClassViewer/Private/SClassViewer.h
@@ -8,6 +8,7 @@
 // SClassViewer
 
 class FClassViewerNode;
+class FTextFilterExpressionEvaluator;
 
 class SClassViewer : public SCompoundWidget
 {
@@ -170,8 +171,8 @@ private:
 	/** The items to be displayed in the tree. */
 	TArray< TSharedPtr< FClassViewerNode > > RootTreeItems;
 
-	/** Split and sanitized list of filter search terms.  Cached for faster search performance. */
-	TArray< FString > FilterSearchTerms;
+	/** Compiled filter search terms. */
+	TSharedPtr<FTextFilterExpressionEvaluator> TextFilterPtr;
 
 	/** Holds the Slate Tree widget which holds the classes for the Class Viewer. */
 	TSharedPtr<STreeView<TSharedPtr< FClassViewerNode > >> ClassTree;
diff --git a/Engine/Source/Editor/ContentBrowser/Private/AssetContextMenu.cpp b/Engine/Source/Editor/ContentBrowser/Private/AssetContextMenu.cpp
--- a/Engine/Source/Editor/ContentBrowser/Private/AssetContextMenu.cpp
+++ b/Engine/Source/Editor/ContentBrowser/Private/AssetContextMenu.cpp
@@ -17,6 +17,7 @@
 #include "ISizeMapModule.h"
 
 #include "ReferencedAssetsUtils.h"
+#include "PackageLocalizationUtil.h"
 
 #include "ISourceControlModule.h"
 #include "ISourceControlRevision.h"
@@ -209,10 +210,26 @@ bool FAssetContextMenu::AddImportedAssetMenuOptions(FMenuBuilder& MenuBuilder)
 
 bool FAssetContextMenu::AddCommonMenuOptions(FMenuBuilder& MenuBuilder)
 {
-	TArray< FAssetData > AssetViewSelectedAssets = AssetView.Pin()->GetSelectedAssets();
-
 	int32 NumAssetItems, NumClassItems;
-	ContentBrowserUtils::CountItemTypes(AssetViewSelectedAssets, NumAssetItems, NumClassItems);
+	ContentBrowserUtils::CountItemTypes(SelectedAssets, NumAssetItems, NumClassItems);
+
+	FAssetToolsModule& AssetToolsModule = FModuleManager::Get().LoadModuleChecked<FAssetToolsModule>("AssetTools");
+
+	// Can any of the selected assets be localized?
+	bool bAnyLocalizableAssetsSelected = false;
+	for (const FAssetData& Asset : SelectedAssets)
+	{
+		TSharedPtr<IAssetTypeActions> AssetTypeActions = AssetToolsModule.Get().GetAssetTypeActionsForClass(Asset.GetClass()).Pin();
+		if (AssetTypeActions.IsValid())
+		{
+			bAnyLocalizableAssetsSelected = AssetTypeActions->CanLocalize();
+		}
+
+		if (bAnyLocalizableAssetsSelected)
+		{
+			break;
+		}
+	}
 
 	MenuBuilder.BeginSection("CommonAssetActions", LOCTEXT("CommonAssetActionsMenuHeading", "Common"));
 	{
@@ -273,6 +290,21 @@ bool FAssetContextMenu::AddCommonMenuOptions(FMenuBuilder& MenuBuilder)
 				false, 
 				FSlateIcon(FEditorStyle::GetStyleSetName(), "ContentBrowser.AssetActions")
 				);
+
+			if (bAnyLocalizableAssetsSelected && NumClassItems == 0)
+			{
+				// Asset Localization sub-menu
+				MenuBuilder.AddSubMenu(
+					LOCTEXT("LocalizationSubMenuLabel", "Asset Localization"),
+					LOCTEXT("LocalizationSubMenuToolTip", "View or create localized variants of this asset"),
+					FNewMenuDelegate::CreateSP(this, &FAssetContextMenu::MakeAssetLocalizationSubMenu),
+					FUIAction(),
+					NAME_None,
+					EUserInterfaceActionType::Button,
+					false,
+					FSlateIcon(FEditorStyle::GetStyleSetName(), "ContentBrowser.AssetLocalization")
+					);
+			}
 		}
 	}
 	MenuBuilder.EndSection();
@@ -470,6 +502,318 @@ bool FAssetContextMenu::CanExecuteAssetActions() const
 	return !bAtLeastOneClassSelected;
 }
 
+void FAssetContextMenu::MakeAssetLocalizationSubMenu(FMenuBuilder& MenuBuilder)
+{
+	TArray<FCultureRef> CurrentCultures;
+
+	// Build up the list of cultures already used
+	{
+		TSet<FString> CulturePaths;
+
+		bool bIncludeEngineCultures = false;
+		bool bIncludeProjectCultures = false;
+
+		for (const FAssetData& Asset : SelectedAssets)
+		{
+			const FString AssetPath = Asset.ObjectPath.ToString();
+
+			if (ContentBrowserUtils::IsEngineFolder(AssetPath))
+			{
+				bIncludeEngineCultures = true;
+			}
+			else
+			{
+				bIncludeProjectCultures = true;
+			}
+
+			{
+				FString AssetLocalizationRoot;
+				if (FPackageLocalizationUtil::GetLocalizedRoot(AssetPath, FString(), AssetLocalizationRoot))
+				{
+					FString AssetLocalizationFileRoot;
+					if (FPackageName::TryConvertLongPackageNameToFilename(AssetLocalizationRoot, AssetLocalizationFileRoot))
+					{
+						CulturePaths.Add(MoveTemp(AssetLocalizationFileRoot));
+					}
+				}
+			}
+		}
+
+		if (bIncludeEngineCultures)
+		{
+			CulturePaths.Append(FPaths::GetEngineLocalizationPaths());
+		}
+
+		if (bIncludeProjectCultures)
+		{
+			CulturePaths.Append(FPaths::GetGameLocalizationPaths());
+		}
+
+		FInternationalization::Get().GetCulturesWithAvailableLocalization(CulturePaths.Array(), CurrentCultures, false);
+
+		if (CurrentCultures.Num() == 0)
+		{
+			CurrentCultures.Add(FInternationalization::Get().GetCurrentCulture());
+		}
+	}
+
+	// Sort by display name for the UI
+	CurrentCultures.Sort([](const FCultureRef& FirstCulture, const FCultureRef& SecondCulture) -> bool
+	{
+		const FText FirstDisplayName = FText::FromString(FirstCulture->GetDisplayName());
+		const FText SecondDisplayName = FText::FromString(SecondCulture->GetDisplayName());
+		return FirstDisplayName.CompareTo(SecondDisplayName) < 0;
+	});
+
+	FAssetToolsModule& AssetToolsModule = FModuleManager::Get().LoadModuleChecked<FAssetToolsModule>("AssetTools");
+
+	// Now build up the list of available localized or source assets based upon the current selection and current cultures
+	FSourceAssetsState SourceAssetsState;
+	TArray<FLocalizedAssetsState> LocalizedAssetsState;
+	for (const FCultureRef& CurrentCulture : CurrentCultures)
+	{
+		FLocalizedAssetsState& LocalizedAssetsStateForCulture = LocalizedAssetsState[LocalizedAssetsState.AddDefaulted()];
+		LocalizedAssetsStateForCulture.Culture = CurrentCulture;
+
+		for (const FAssetData& Asset : SelectedAssets)
+		{
+			// Can this type of asset be localized?
+			bool bCanLocalizeAsset = false;
+			{
+				TSharedPtr<IAssetTypeActions> AssetTypeActions = AssetToolsModule.Get().GetAssetTypeActionsForClass(Asset.GetClass()).Pin();
+				if (AssetTypeActions.IsValid())
+				{
+					bCanLocalizeAsset = AssetTypeActions->CanLocalize();
+				}
+			}
+
+			if (!bCanLocalizeAsset)
+			{
+				continue;
+			}
+
+			const FString ObjectPath = Asset.ObjectPath.ToString();
+			if (FPackageName::IsLocalizedPackage(ObjectPath))
+			{
+				// Get the source path for this asset
+				FString SourceObjectPath;
+				if (FPackageLocalizationUtil::ConvertLocalizedToSource(ObjectPath, SourceObjectPath))
+				{
+					SourceAssetsState.CurrentAssets.Add(*SourceObjectPath);
+				}
+			}
+			else
+			{
+				SourceAssetsState.SelectedAssets.Add(Asset.ObjectPath);
+
+				// Get the localized path for this asset and culture
+				FString LocalizedObjectPath;
+				if (FPackageLocalizationUtil::ConvertSourceToLocalized(ObjectPath, CurrentCulture->GetName(), LocalizedObjectPath))
+				{
+					// Does this localized asset already exist?
+					FAssetRegistryModule& AssetRegistryModule = FModuleManager::LoadModuleChecked<FAssetRegistryModule>(TEXT("AssetRegistry"));
+					FAssetData LocalizedAssetData = AssetRegistryModule.Get().GetAssetByObjectPath(*LocalizedObjectPath);
+
+					if (LocalizedAssetData.IsValid())
+					{
+						LocalizedAssetsStateForCulture.CurrentAssets.Add(*LocalizedObjectPath);
+					}
+					else
+					{
+						LocalizedAssetsStateForCulture.NewAssets.Add(*LocalizedObjectPath);
+					}
+				}
+			}
+		}
+	}
+
+	// If we found source assets for localized assets, then we can show the Source Asset options
+	if (SourceAssetsState.CurrentAssets.Num() > 0)
+	{
+		MenuBuilder.BeginSection(NAME_None, LOCTEXT("ManageSourceAssetHeading", "Manage Source Asset"));
+		{
+			MenuBuilder.AddMenuEntry(
+				LOCTEXT("ShowSourceAsset", "Show Source Asset"),
+				LOCTEXT("ShowSourceAssetTooltip", "Show the source asset in the Content Browser."),
+				FSlateIcon(FEditorStyle::GetStyleSetName(), "SystemWideCommands.FindInContentBrowser"),
+				FUIAction(FExecuteAction::CreateSP(this, &FAssetContextMenu::ExecuteFindInAssetTree, SourceAssetsState.CurrentAssets.Array()))
+				);
+
+			MenuBuilder.AddMenuEntry(
+				LOCTEXT("EditSourceAsset", "Edit Source Asset"),
+				LOCTEXT("EditSourceAssetTooltip", "Edit the source asset."),
+				FSlateIcon(FEditorStyle::GetStyleSetName(), "ContentBrowser.AssetActions.Edit"),
+				FUIAction(FExecuteAction::CreateSP(this, &FAssetContextMenu::ExecuteOpenEditorsForAssets, SourceAssetsState.CurrentAssets.Array()))
+				);
+		}
+		MenuBuilder.EndSection();
+	}
+
+	// If we currently have source assets selected, then we can show the Localized Asset options
+	if (SourceAssetsState.SelectedAssets.Num() > 0)
+	{
+		MenuBuilder.BeginSection(NAME_None, LOCTEXT("ManageLocalizedAssetHeading", "Manage Localized Asset"));
+		{
+			MenuBuilder.AddSubMenu(
+				LOCTEXT("CreateLocalizedAsset", "Create Localized Asset"),
+				LOCTEXT("CreateLocalizedAssetTooltip", "Create a new localized asset."),
+				FNewMenuDelegate::CreateSP(this, &FAssetContextMenu::MakeCreateLocalizedAssetSubMenu, SourceAssetsState.SelectedAssets, LocalizedAssetsState),
+				FUIAction(),
+				NAME_None,
+				EUserInterfaceActionType::Button,
+				false,
+				FSlateIcon(FEditorStyle::GetStyleSetName(), "ContentBrowser.AssetActions.Duplicate")
+				);
+
+			int32 NumLocalizedAssets = 0;
+			for (const FLocalizedAssetsState& LocalizedAssetsStateForCulture : LocalizedAssetsState)
+			{
+				NumLocalizedAssets += LocalizedAssetsStateForCulture.CurrentAssets.Num();
+			}
+
+			if (NumLocalizedAssets > 0)
+			{
+				MenuBuilder.AddSubMenu(
+					LOCTEXT("ShowLocalizedAsset", "Show Localized Asset"),
+					LOCTEXT("ShowLocalizedAssetTooltip", "Show the localized asset in the Content Browser."),
+					FNewMenuDelegate::CreateSP(this, &FAssetContextMenu::MakeShowLocalizedAssetSubMenu, LocalizedAssetsState),
+					FUIAction(),
+					NAME_None,
+					EUserInterfaceActionType::Button,
+					false,
+					FSlateIcon(FEditorStyle::GetStyleSetName(), "SystemWideCommands.FindInContentBrowser")
+					);
+
+				MenuBuilder.AddSubMenu(
+					LOCTEXT("EditLocalizedAsset", "Edit Localized Asset"),
+					LOCTEXT("EditLocalizedAssetTooltip", "Edit the localized asset."),
+					FNewMenuDelegate::CreateSP(this, &FAssetContextMenu::MakeEditLocalizedAssetSubMenu, LocalizedAssetsState),
+					FUIAction(),
+					NAME_None,
+					EUserInterfaceActionType::Button,
+					false,
+					FSlateIcon(FEditorStyle::GetStyleSetName(), "ContentBrowser.AssetActions.Edit")
+					);
+			}
+		}
+		MenuBuilder.EndSection();
+	}
+}
+
+void FAssetContextMenu::MakeCreateLocalizedAssetSubMenu(FMenuBuilder& MenuBuilder, TSet<FName> InSelectedSourceAssets, TArray<FLocalizedAssetsState> InLocalizedAssetsState)
+{
+	for (const FLocalizedAssetsState& LocalizedAssetsStateForCulture : InLocalizedAssetsState)
+	{
+		// If we have less localized assets than we have selected source assets, then we'll have some assets to create localized variants of
+		if (LocalizedAssetsStateForCulture.CurrentAssets.Num() < InSelectedSourceAssets.Num())
+		{
+			MenuBuilder.AddMenuEntry(
+				FText::FromString(LocalizedAssetsStateForCulture.Culture->GetDisplayName()),
+				FText::GetEmpty(),
+				FSlateIcon(),
+				FUIAction(FExecuteAction::CreateSP(this, &FAssetContextMenu::ExecuteCreateLocalizedAsset, InSelectedSourceAssets, LocalizedAssetsStateForCulture))
+				);
+		}
+	}
+}
+
+void FAssetContextMenu::MakeShowLocalizedAssetSubMenu(FMenuBuilder& MenuBuilder, TArray<FLocalizedAssetsState> InLocalizedAssetsState)
+{
+	for (const FLocalizedAssetsState& LocalizedAssetsStateForCulture : InLocalizedAssetsState)
+	{
+		if (LocalizedAssetsStateForCulture.CurrentAssets.Num() > 0)
+		{
+			MenuBuilder.AddMenuEntry(
+				FText::FromString(LocalizedAssetsStateForCulture.Culture->GetDisplayName()),
+				FText::GetEmpty(),
+				FSlateIcon(),
+				FUIAction(FExecuteAction::CreateSP(this, &FAssetContextMenu::ExecuteFindInAssetTree, LocalizedAssetsStateForCulture.CurrentAssets.Array()))
+				);
+		}
+	}
+}
+
+void FAssetContextMenu::MakeEditLocalizedAssetSubMenu(FMenuBuilder& MenuBuilder, TArray<FLocalizedAssetsState> InLocalizedAssetsState)
+{
+	for (const FLocalizedAssetsState& LocalizedAssetsStateForCulture : InLocalizedAssetsState)
+	{
+		if (LocalizedAssetsStateForCulture.CurrentAssets.Num() > 0)
+		{
+			MenuBuilder.AddMenuEntry(
+				FText::FromString(LocalizedAssetsStateForCulture.Culture->GetDisplayName()),
+				FText::GetEmpty(),
+				FSlateIcon(),
+				FUIAction(FExecuteAction::CreateSP(this, &FAssetContextMenu::ExecuteOpenEditorsForAssets, LocalizedAssetsStateForCulture.CurrentAssets.Array()))
+				);
+		}
+	}
+}
+
+void FAssetContextMenu::ExecuteCreateLocalizedAsset(TSet<FName> InSelectedSourceAssets, FLocalizedAssetsState InLocalizedAssetsStateForCulture)
+{
+	TArray<UPackage*> PackagesToSave;
+	TArray<FAssetData> NewObjects;
+
+	for (const FName SourceAssetName : InSelectedSourceAssets)
+	{
+		if (InLocalizedAssetsStateForCulture.CurrentAssets.Contains(SourceAssetName))
+		{
+			// Asset is already localized
+			continue;
+		}
+
+		UObject* SourceAssetObject = LoadObject<UObject>(nullptr, *SourceAssetName.ToString());
+		if (!SourceAssetObject)
+		{
+			// Source object cannot be loaded
+			continue;
+		}
+
+		FString LocalizedPackageName;
+		if (!FPackageLocalizationUtil::ConvertSourceToLocalized(SourceAssetObject->GetOutermost()->GetPathName(), InLocalizedAssetsStateForCulture.Culture->GetName(), LocalizedPackageName))
+		{
+			continue;
+		}
+
+		ObjectTools::FPackageGroupName NewAssetName;
+		NewAssetName.PackageName = LocalizedPackageName;
+		NewAssetName.ObjectName = SourceAssetObject->GetName();
+
+		TSet<UPackage*> PackagesNotDuplicated;
+		UObject* NewObject = ObjectTools::DuplicateSingleObject(SourceAssetObject, NewAssetName, PackagesNotDuplicated);
+		if (NewObject)
+		{
+			PackagesToSave.Add(NewObject->GetOutermost());
+			NewObjects.Add(FAssetData(NewObject));
+		}
+	}
+
+	if (PackagesToSave.Num() > 0)
+	{
+		FEditorFileUtils::PromptForCheckoutAndSave(PackagesToSave, /*bCheckDirty*/false, /*bPromptToSave*/false);
+	}
+
+	OnFindInAssetTreeRequested.ExecuteIfBound(NewObjects);
+}
+
+void FAssetContextMenu::ExecuteFindInAssetTree(TArray<FName> InAssets)
+{
+	FAssetRegistryModule& AssetRegistryModule = FModuleManager::LoadModuleChecked<FAssetRegistryModule>(TEXT("AssetRegistry"));
+
+	FARFilter ARFilter;
+	ARFilter.ObjectPaths = MoveTemp(InAssets);
+	
+	TArray<FAssetData> FoundLocalizedAssetData;
+	AssetRegistryModule.Get().GetAssets(ARFilter, FoundLocalizedAssetData);
+
+	OnFindInAssetTreeRequested.ExecuteIfBound(FoundLocalizedAssetData);
+}
+
+void FAssetContextMenu::ExecuteOpenEditorsForAssets(TArray<FName> InAssets)
+{
+	FAssetEditorManager::Get().OpenEditorsForAssets(InAssets);
+}
+
 bool FAssetContextMenu::AddReferenceMenuOptions(FMenuBuilder& MenuBuilder)
 {
 	MenuBuilder.BeginSection("AssetContextReferences", LOCTEXT("ReferencesMenuHeading", "References"));
diff --git a/Engine/Source/Editor/ContentBrowser/Private/AssetContextMenu.h b/Engine/Source/Editor/ContentBrowser/Private/AssetContextMenu.h
--- a/Engine/Source/Editor/ContentBrowser/Private/AssetContextMenu.h
+++ b/Engine/Source/Editor/ContentBrowser/Private/AssetContextMenu.h
@@ -48,6 +48,20 @@ public:
 	/** Handler for Delete */
 	void ExecuteDelete();
 
+private:
+	struct FSourceAssetsState
+	{
+		TSet<FName> SelectedAssets;
+		TSet<FName> CurrentAssets;
+	};
+
+	struct FLocalizedAssetsState
+	{
+		FCulturePtr Culture;
+		TSet<FName> NewAssets;
+		TSet<FName> CurrentAssets;
+	};
+
 private:
 	/** Helper to load selected assets and sort them by UClass */
 	void GetSelectedAssetsByClass(TMap<UClass*, TArray<UObject*> >& OutSelectedAssetsByClass) const;
@@ -95,6 +109,27 @@ private:
 	/** Handler to check to see if "Asset Actions" are allowed */
 	bool CanExecuteAssetActions() const;
 
+	/** Adds Asset Localization sub-menu to a menu builder. */
+	void MakeAssetLocalizationSubMenu(FMenuBuilder& MenuBuilder);
+
+	/** Adds the Create Localized Asset sub-menu to a menu builder. */
+	void MakeCreateLocalizedAssetSubMenu(FMenuBuilder& MenuBuilder, TSet<FName> InSelectedSourceAssets, TArray<FLocalizedAssetsState> InLocalizedAssetsState);
+
+	/** Adds the Show Localized Assets sub-menu to a menu builder. */
+	void MakeShowLocalizedAssetSubMenu(FMenuBuilder& MenuBuilder, TArray<FLocalizedAssetsState> InLocalizedAssetsState);
+
+	/** Adds the Edit Localized Assets sub-menu to a menu builder. */
+	void MakeEditLocalizedAssetSubMenu(FMenuBuilder& MenuBuilder, TArray<FLocalizedAssetsState> InLocalizedAssetsState);
+
+	/** Create new localized assets for the given culture */
+	void ExecuteCreateLocalizedAsset(TSet<FName> InSelectedSourceAssets, FLocalizedAssetsState InLocalizedAssetsStateForCulture);
+
+	/** Find the given assets in the Content Browser */
+	void ExecuteFindInAssetTree(TArray<FName> InAssets);
+
+	/** Open the given assets in their respective editors */
+	void ExecuteOpenEditorsForAssets(TArray<FName> InAssets);
+
 	/** Adds asset reference menu options to a menu builder. Returns true if any options were added. */
 	bool AddReferenceMenuOptions(FMenuBuilder& MenuBuilder);
 
diff --git a/Engine/Source/Editor/ContentBrowser/Private/SAssetView.cpp b/Engine/Source/Editor/ContentBrowser/Private/SAssetView.cpp
--- a/Engine/Source/Editor/ContentBrowser/Private/SAssetView.cpp
+++ b/Engine/Source/Editor/ContentBrowser/Private/SAssetView.cpp
@@ -1373,41 +1373,42 @@ static bool IsValidObjectPath(const FString& Path)
 
 FReply SAssetView::OnKeyDown( const FGeometry& MyGeometry, const FKeyEvent& InKeyEvent )
 {
+	if (InKeyEvent.IsControlDown() && InKeyEvent.GetCharacter() == 'V' && IsAssetPathSelected())
 	{
-		const bool bTestOnly = true;
-		if ( InKeyEvent.IsControlDown() && InKeyEvent.GetCharacter() == 'V' && IsAssetPathSelected() )
-		{
-			FString DestPaths;
-			TArray<FString> DestPathsSplit;
+		FString AssetPaths;
+		TArray<FString> AssetPathsSplit;
 
-			// Get the copied asset paths
-			FPlatformMisc::ClipboardPaste( DestPaths );
-			DestPaths.ParseIntoArrayLines( DestPathsSplit );
+		// Get the copied asset paths
+		FPlatformMisc::ClipboardPaste(AssetPaths);
+		AssetPaths.ParseIntoArrayLines(AssetPathsSplit);
 
-			// Get assets and copy them
-			TArray<UObject*> ObjectsToCopy;
-			for (FString DestPath : DestPathsSplit)
+		// Get assets and copy them
+		TArray<UObject*> AssetsToCopy;
+		for (const FString& AssetPath : AssetPathsSplit)
+		{
+			// Validate string
+			if (IsValidObjectPath(AssetPath))
 			{
-				// Validate string
-				if ( IsValidObjectPath(DestPath) )
+				UObject* ObjectToCopy = LoadObject<UObject>(nullptr, *AssetPath);
+				if (ObjectToCopy && !ObjectToCopy->IsA(UClass::StaticClass()))
 				{
-					ObjectsToCopy.Add( LoadObject<UObject>( NULL, *DestPath ));
+					AssetsToCopy.Add(ObjectToCopy);
 				}
 			}
-
-			if (ObjectsToCopy.Num())
-			{
-				ContentBrowserUtils::CopyAssets(ObjectsToCopy, SourcesData.PackagePaths[0].ToString());
-			}
-
-			return FReply::Handled();
 		}
-		// Swallow the key-presses used by the quick-jump in OnKeyChar to avoid other things (such as the viewport commands) getting them instead
-		// eg) Pressing "W" without this would set the viewport to "translate" mode
-		else if( HandleQuickJumpKeyDown(InKeyEvent.GetCharacter(), InKeyEvent.IsControlDown(), InKeyEvent.IsAltDown(), bTestOnly).IsEventHandled() )
+
+		if (AssetsToCopy.Num())
 		{
-			return FReply::Handled();
+			ContentBrowserUtils::CopyAssets(AssetsToCopy, SourcesData.PackagePaths[0].ToString());
 		}
+
+		return FReply::Handled();
+	}
+	// Swallow the key-presses used by the quick-jump in OnKeyChar to avoid other things (such as the viewport commands) getting them instead
+	// eg) Pressing "W" without this would set the viewport to "translate" mode
+	else if(HandleQuickJumpKeyDown(InKeyEvent.GetCharacter(), InKeyEvent.IsControlDown(), InKeyEvent.IsAltDown(), /*bTestOnly*/true).IsEventHandled())
+	{
+		return FReply::Handled();
 	}
 
 	return FReply::Unhandled();
diff --git a/Engine/Source/Editor/ContentBrowser/Private/SContentBrowser.cpp b/Engine/Source/Editor/ContentBrowser/Private/SContentBrowser.cpp
--- a/Engine/Source/Editor/ContentBrowser/Private/SContentBrowser.cpp
+++ b/Engine/Source/Editor/ContentBrowser/Private/SContentBrowser.cpp
@@ -789,10 +789,11 @@ void SContentBrowser::SyncToAssets( const TArray<FAssetData>& AssetDataList, con
 	bool bDisplayDev = GetDefault<UContentBrowserSettings>()->GetDisplayDevelopersFolder();
 	bool bDisplayEngine = GetDefault<UContentBrowserSettings>()->GetDisplayEngineFolder();
 	bool bDisplayPlugins = GetDefault<UContentBrowserSettings>()->GetDisplayPluginFolders();
-	if ( !bDisplayDev || !bDisplayEngine || !bDisplayPlugins )
+	bool bDisplayLocalized = GetDefault<UContentBrowserSettings>()->GetDisplayL10NFolder();
+	if ( !bDisplayDev || !bDisplayEngine || !bDisplayPlugins || !bDisplayLocalized )
 	{
 		bool bRepopulate = false;
-		for (int32 AssetIdx = AssetDataList.Num() - 1; AssetIdx >= 0 && ( !bDisplayDev || !bDisplayEngine || !bDisplayPlugins ); --AssetIdx)
+		for (int32 AssetIdx = AssetDataList.Num() - 1; AssetIdx >= 0 && ( !bDisplayDev || !bDisplayEngine || !bDisplayPlugins || !bDisplayLocalized ); --AssetIdx)
 		{
 			const FAssetData& Item = AssetDataList[AssetIdx];
 
@@ -828,6 +829,13 @@ void SContentBrowser::SyncToAssets( const TArray<FAssetData>& AssetDataList, con
 				GetMutableDefault<UContentBrowserSettings>()->SetDisplayPluginFolders(true, true);
 				bRepopulate = true;
 			}
+
+			if (!bDisplayLocalized && ContentBrowserUtils::IsLocalizationFolder(PackagePath))
+			{
+				bDisplayLocalized = true;
+				GetMutableDefault<UContentBrowserSettings>()->SetDisplayL10NFolder(true);
+				bRepopulate = true;
+			}
 		}
 
 		// If we have auto-enabled any flags, force a refresh
diff --git a/Engine/Source/Editor/ContentBrowser/Private/SPathView.cpp b/Engine/Source/Editor/ContentBrowser/Private/SPathView.cpp
--- a/Engine/Source/Editor/ContentBrowser/Private/SPathView.cpp
+++ b/Engine/Source/Editor/ContentBrowser/Private/SPathView.cpp
@@ -638,10 +638,11 @@ void SPathView::LoadSettings(const FString& IniFilename, const FString& IniSecti
 
 			// Send the first selected item with the notification
 			const TArray<TSharedPtr<FTreeItem>> SelectedItems = TreeViewPtr->GetSelectedItems();
-			check(SelectedItems.Num() > 0);
-
-			// Signal a single selection changed event to let any listeners know that paths have changed
-			TreeSelectionChanged( SelectedItems[0], ESelectInfo::Direct );
+			if (SelectedItems.Num() > 0)
+			{
+				// Signal a single selection changed event to let any listeners know that paths have changed
+				TreeSelectionChanged( SelectedItems[0], ESelectInfo::Direct );
+			}
 		}
 	}
 }
diff --git a/Engine/Source/Editor/DetailCustomizations/Private/CurveStructCustomization.cpp b/Engine/Source/Editor/DetailCustomizations/Private/CurveStructCustomization.cpp
--- a/Engine/Source/Editor/DetailCustomizations/Private/CurveStructCustomization.cpp
+++ b/Engine/Source/Editor/DetailCustomizations/Private/CurveStructCustomization.cpp
@@ -41,7 +41,6 @@ void FCurveStructCustomization::CustomizeHeader( TSharedRef<IPropertyHandle> InS
 
 	TArray<void*> StructPtrs;
 	StructPropertyHandle->AccessRawData( StructPtrs );
-	check(StructPtrs.Num()!=0);
 
 	if (StructPtrs.Num() == 1)
 	{
@@ -70,7 +69,7 @@ void FCurveStructCustomization::CustomizeHeader( TSharedRef<IPropertyHandle> InS
 		HeaderRow
 			.NameContent()
 			[
-				InStructPropertyHandle->CreatePropertyNameWidget( FText::GetEmpty(), FText::GetEmpty(), false )
+				InStructPropertyHandle->CreatePropertyNameWidget()
 			]
 			.ValueContent()
 			.HAlign(HAlign_Fill)
@@ -107,7 +106,7 @@ void FCurveStructCustomization::CustomizeHeader( TSharedRef<IPropertyHandle> InS
 		HeaderRow
 			.NameContent()
 			[
-				InStructPropertyHandle->CreatePropertyNameWidget( FText::GetEmpty(), FText::GetEmpty(), false )
+				InStructPropertyHandle->CreatePropertyNameWidget()
 			]
 			.ValueContent()
 			[
@@ -115,7 +114,7 @@ void FCurveStructCustomization::CustomizeHeader( TSharedRef<IPropertyHandle> InS
 				.VAlign(VAlign_Fill)
 				[
 					SNew(STextBlock)
-					.Text(LOCTEXT("MultipleCurves", "Multiple Curves - unable to modify"))
+					.Text(StructPtrs.Num() == 0 ? LOCTEXT("NoCurves", "No Curves - unable to modify") : LOCTEXT("MultipleCurves", "Multiple Curves - unable to modify"))
 				]
 			];
 	}
diff --git a/Engine/Source/Editor/DetailCustomizations/Private/TextCustomization.cpp b/Engine/Source/Editor/DetailCustomizations/Private/TextCustomization.cpp
--- a/Engine/Source/Editor/DetailCustomizations/Private/TextCustomization.cpp
+++ b/Engine/Source/Editor/DetailCustomizations/Private/TextCustomization.cpp
@@ -46,33 +46,38 @@ namespace
 
 		virtual int32 GetNumTexts() const override
 		{
-			const TArray<FText*>& RawTextData = SyncRawTextData();
-
 			return (PropertyHandle->IsValidHandle())
-				? RawTextData.Num() 
+				? PropertyHandle->GetNumPerObjectValues() 
 				: 0;
 		}
 
 		virtual FText GetText(const int32 InIndex) const override
 		{
-			const TArray<FText*>& RawTextData = SyncRawTextData();
-
 			if (PropertyHandle->IsValidHandle())
 			{
-				check(RawTextData.IsValidIndex(InIndex));
-				return *RawTextData[InIndex];
+				FString ObjectValue;
+				if (PropertyHandle->GetPerObjectValue(InIndex, ObjectValue) == FPropertyAccess::Success)
+				{
+					FText TextValue;
+					if (FTextStringHelper::ReadFromString(*ObjectValue, TextValue))
+					{
+						return TextValue;
+					}
+				}
 			}
+
 			return FText::GetEmpty();
 		}
 
 		virtual void SetText(const int32 InIndex, const FText& InText) override
 		{
-			const TArray<FText*>& RawTextData = SyncRawTextData();
-
 			if (PropertyHandle->IsValidHandle())
 			{
-				check(RawTextData.IsValidIndex(InIndex));
-				*RawTextData[InIndex] = InText;
+				FString ObjectValue;
+				if (FTextStringHelper::WriteToString(ObjectValue, InText))
+				{
+					PropertyHandle->SetPerObjectValue(InIndex, ObjectValue);
+				}
 			}
 		}
 
@@ -102,26 +107,8 @@ namespace
 		}
 
 	private:
-		const TArray<FText*>& SyncRawTextData() const
-		{
-			// Sync the scratch data array with the current property data
-			RawTextDataScratchArray.Reset();
-
-			if (PropertyHandle->IsValidHandle())
-			{
-				PropertyHandle->EnumerateRawData([&](void* RawData, const int32 DataIndex, const int32 NumDatas) -> bool
-				{
-					RawTextDataScratchArray.Add(static_cast<FText*>(RawData));
-					return true;
-				});
-			}
-
-			return RawTextDataScratchArray;
-		}
-
 		TSharedRef<IPropertyHandle> PropertyHandle;
 		TSharedPtr<IPropertyUtilities> PropertyUtilities;
-		mutable TArray<FText*> RawTextDataScratchArray;
 	};
 }
 
diff --git a/Engine/Source/Editor/DistCurveEditor/Private/CurveEditorViewportClient.cpp b/Engine/Source/Editor/DistCurveEditor/Private/CurveEditorViewportClient.cpp
--- a/Engine/Source/Editor/DistCurveEditor/Private/CurveEditorViewportClient.cpp
+++ b/Engine/Source/Editor/DistCurveEditor/Private/CurveEditorViewportClient.cpp
@@ -403,10 +403,10 @@ bool FCurveEditorViewportClient::InputKey(FViewport* Viewport, int32 ControllerI
 				{
 					if(bBoxSelecting)
 					{
-						int32 MinX = FMath::Min(BoxStartX, BoxEndX);
-						int32 MinY = FMath::Min(BoxStartY, BoxEndY);
-						int32 MaxX = FMath::Max(BoxStartX, BoxEndX);
-						int32 MaxY = FMath::Max(BoxStartY, BoxEndY);
+						int32 MinX = FMath::Max(0, FMath::Min(BoxStartX, BoxEndX));
+						int32 MinY = FMath::Max(0, FMath::Min(BoxStartY, BoxEndY));
+						int32 MaxX = FMath::Min(Viewport->GetSizeXY().X - 1, FMath::Max(BoxStartX, BoxEndX));
+						int32 MaxY = FMath::Min(Viewport->GetSizeXY().Y - 1, FMath::Max(BoxStartY, BoxEndY));
 						int32 TestSizeX = MaxX - MinX + 1;
 						int32 TestSizeY = MaxY - MinY + 1;
 
diff --git a/Engine/Source/Editor/EditorStyle/Private/SlateEditorStyle.cpp b/Engine/Source/Editor/EditorStyle/Private/SlateEditorStyle.cpp
--- a/Engine/Source/Editor/EditorStyle/Private/SlateEditorStyle.cpp
+++ b/Engine/Source/Editor/EditorStyle/Private/SlateEditorStyle.cpp
@@ -1524,7 +1524,8 @@ void FSlateEditorStyle::FStyle::SetupGeneralStyles()
 			.SetNormalBarImage(FSlateColorBrush(FColor::White))
 			.SetDisabledBarImage(FSlateColorBrush(FLinearColor::Gray))
 			.SetNormalThumbImage( BOX_BRUSH( "Common/Button", 8.0f/32.0f ) )
-			.SetDisabledThumbImage( BOX_BRUSH( "Common/Button_Disabled", 8.0f/32.0f ) );
+			.SetDisabledThumbImage( BOX_BRUSH( "Common/Button_Disabled", 8.0f/32.0f ) )
+			.SetBarThickness(2.0f);
 		Set( "Slider", SliderStyle );
 
 		Set( "VolumeControl", FVolumeControlStyle()
@@ -5838,6 +5839,7 @@ void FSlateEditorStyle::FStyle::SetupContentBrowserStyle()
 		Set( "ContentBrowser.AssetActions.CreateThumbnail", new IMAGE_BRUSH( "Icons/icon_Asset_Create_Thumbnail_16x", Icon16x16) );
 		Set( "ContentBrowser.AssetActions.DeleteThumbnail", new IMAGE_BRUSH( "Icons/icon_Asset_Delete_Thumbnail_16x", Icon16x16) );
 		Set( "ContentBrowser.AssetActions.GenericFind", new IMAGE_BRUSH( "Icons/icon_Genericfinder_16x", Icon16x16) );
+		Set( "ContentBrowser.AssetLocalization", new IMAGE_BRUSH( "Icons/icon_localization_16x", Icon16x16 ) );
 
 		Set( "MediaAsset.AssetActions.Play", new IMAGE_BRUSH( "Icons/icon_SCueEd_PlayCue_16x", Icon16x16 ) );
 		Set( "MediaAsset.AssetActions.Stop", new IMAGE_BRUSH( "Icons/icon_SCueEd_Stop_16x", Icon16x16 ) );
diff --git a/Engine/Source/Editor/GraphEditor/Private/KismetNodes/SGraphNodeK2CreateDelegate.cpp b/Engine/Source/Editor/GraphEditor/Private/KismetNodes/SGraphNodeK2CreateDelegate.cpp
--- a/Engine/Source/Editor/GraphEditor/Private/KismetNodes/SGraphNodeK2CreateDelegate.cpp
+++ b/Engine/Source/Editor/GraphEditor/Private/KismetNodes/SGraphNodeK2CreateDelegate.cpp
@@ -56,7 +56,7 @@ FText SGraphNodeK2CreateDelegate::GetCurrentFunctionDescription() const
 
 	if(!FunctionSignature || !ScopeClass)
 	{
-		return NSLOCTEXT("GraphNodeK2Create", "NoneLabel", "None");
+		return NSLOCTEXT("GraphNodeK2Create", "NoneLabel", "");
 	}
 
 	if (const auto Func = FindField<UFunction>(ScopeClass, Node->GetFunctionName()))
@@ -69,7 +69,7 @@ FText SGraphNodeK2CreateDelegate::GetCurrentFunctionDescription() const
 		return FText::Format(NSLOCTEXT("GraphNodeK2Create", "ErrorLabelFmt", "Error? {0}"), FText::FromName(Node->GetFunctionName()));
 	}
 
-	return NSLOCTEXT("GraphNodeK2Create", "SelectFunctionLabel", "Select Function");
+	return NSLOCTEXT("GraphNodeK2Create", "SelectFunctionLabel", "Select Function...");
 }
 
 TSharedRef<ITableRow> SGraphNodeK2CreateDelegate::HandleGenerateRowFunction(TSharedPtr<FFunctionItemData> FunctionItemData, const TSharedRef<STableViewBase>& OwnerTable)
@@ -109,6 +109,7 @@ void SGraphNodeK2CreateDelegate::CreateBelowWidgetControls(TSharedPtr<SVerticalB
 		if(FunctionSignature && ScopeClass)
 		{
 			FunctionDataItems.Empty();
+
 			for(TFieldIterator<UFunction> It(ScopeClass); It; ++It)
 			{
 				UFunction* Func = *It;
@@ -122,6 +123,14 @@ void SGraphNodeK2CreateDelegate::CreateBelowWidgetControls(TSharedPtr<SVerticalB
 				}
 			}
 
+			if (FunctionDataItems.Num() == 0)
+			{
+				// add an empty row, so the user can clear the selection if they want
+				TSharedPtr<FFunctionItemData> EmptyEntry = MakeShareable(new FFunctionItemData());
+				EmptyEntry->Description = NSLOCTEXT("GraphNodeK2Create", "EmptyFunctionLabel", "[NONE]").ToString();
+				FunctionDataItems.Add(EmptyEntry);
+			}
+
 			TSharedRef<SComboButton> SelectFunctionWidgetRef = SNew(SComboButton)
 				.Method(EPopupMethod::UseCurrentWindow)
 				.ButtonContent()
diff --git a/Engine/Source/Editor/KismetCompiler/Private/AnimBlueprintCompiler.cpp b/Engine/Source/Editor/KismetCompiler/Private/AnimBlueprintCompiler.cpp
--- a/Engine/Source/Editor/KismetCompiler/Private/AnimBlueprintCompiler.cpp
+++ b/Engine/Source/Editor/KismetCompiler/Private/AnimBlueprintCompiler.cpp
@@ -121,12 +121,6 @@ UK2Node_CallFunction* FAnimBlueprintCompiler::SpawnCallAnimInstanceFunction(UEdG
 
 void FAnimBlueprintCompiler::CreateEvaluationHandler(UAnimGraphNode_Base* VisualAnimNode, FEvaluationHandlerRecord& Record)
 {
-	if(Record.OnlyUsesCopyRecords())
-	{
-		// simple property copies don't require wiring up to a custom event
-		return;
-	}
-
 	// Shouldn't create a handler if there is nothing to work with
 	check(Record.ServicedProperties.Num() > 0);
 	check(Record.NodeVariableProperty != NULL);
@@ -1208,6 +1202,9 @@ void FAnimBlueprintCompiler::MergeUbergraphPagesIn(UEdGraph* Ubergraph)
 			}
 		}
 
+		// Make sure we expand any split pins here before we process animation nodes.
+		ExpansionStep(ConsolidatedEventGraph, false);
+
 		// Compile the animation graph
 		ProcessAllAnimationNodes();
 	}
@@ -1365,7 +1362,9 @@ void FAnimBlueprintCompiler::PostCompile()
 			const FAnimNodeSinglePropertyHandler& Handler = EvaluationHandler.ServicedProperties.CreateConstIterator()->Value;
 			UAnimGraphNode_Base* Node = CastChecked<UAnimGraphNode_Base>(Handler.ArrayPins.Num() > 0 ? Handler.ArrayPins[0]->GetOwningNode() : Handler.SinglePin->GetOwningNode());
 			UAnimGraphNode_Base* TrueNode = MessageLog.FindSourceObjectTypeChecked<UAnimGraphNode_Base>(Node);	
-			TrueNode->BlueprintUsage = EvaluationHandler.HandlerFunctionName != NAME_None ? EBlueprintUsage::UsesBlueprint : EBlueprintUsage::DoesNotUseBlueprint;
+
+			FExposedValueHandler* HandlerPtr = EvaluationHandler.EvaluationHandlerProperty->ContainerPtrToValuePtr<FExposedValueHandler>(EvaluationHandler.NodeVariableProperty->ContainerPtrToValuePtr<void>(DefaultAnimInstance));
+			TrueNode->BlueprintUsage = HandlerPtr->BoundFunction != NAME_None ? EBlueprintUsage::UsesBlueprint : EBlueprintUsage::DoesNotUseBlueprint;
 
 			if(TrueNode->BlueprintUsage == EBlueprintUsage::UsesBlueprint && DefaultAnimInstance->bWarnAboutBlueprintUsage)
 			{
@@ -1897,26 +1896,56 @@ void FAnimBlueprintCompiler::FEvaluationHandlerRecord::PatchFunctionNameAndCopyR
 			// we dont support copying *from* array properties at the moment
 			if (SimpleCopyPropertySource && !SimpleCopyPropertySource->IsA<UArrayProperty>())
 			{
+				// are we copying to an array element (or more than one)?
 				if (AnimNodeSinglePropertyHandler.ArrayPins.Num() > 0)
 				{
 					UArrayProperty* SimpleCopyPropertyDest = CastChecked<UArrayProperty>(NodeVariableProperty->Struct->FindPropertyByName(PropertyIt.Key()));
-					if(SimpleCopyPropertySource->GetSize() != SimpleCopyPropertyDest->Inner->GetSize())
+
+					if (AnimNodeSinglePropertyHandler.SubStructPropertyName != NAME_None)
 					{
-						bOnlyUsesCopyRecords = false;
-						break;
-					}
+						UStructProperty* SourceStructProperty = CastChecked<UStructProperty>(SimpleCopyPropertySource);
+						UProperty* SourceSubProperty = SourceStructProperty->Struct->FindPropertyByName(AnimNodeSinglePropertyHandler.SubStructPropertyName);
+						if (SourceSubProperty == nullptr || SourceSubProperty->GetSize() != SimpleCopyPropertyDest->Inner->GetSize())
+						{
+							bOnlyUsesCopyRecords = false;
+							break;
+						}
 
-					for (TMap<int32, UEdGraphPin*>::TConstIterator ArrayIt(AnimNodeSinglePropertyHandler.ArrayPins); ArrayIt; ++ArrayIt)
+						for (TMap<int32, UEdGraphPin*>::TConstIterator ArrayIt(AnimNodeSinglePropertyHandler.ArrayPins); ArrayIt; ++ArrayIt)
+						{
+							// Single sub-struct variable get->array set
+							FExposedValueCopyRecord CopyRecord;
+							CopyRecord.DestProperty = SimpleCopyPropertyDest;
+							CopyRecord.DestArrayIndex = ArrayIt.Key();
+							CopyRecord.SourcePropertyName = AnimNodeSinglePropertyHandler.SimpleCopyPropertyName;
+							CopyRecord.SourceSubPropertyName = AnimNodeSinglePropertyHandler.SubStructPropertyName;
+							CopyRecord.SourceArrayIndex = 0;
+							CopyRecord.Size = SimpleCopyPropertyDest->Inner->GetSize();
+							CopyRecord.PostCopyOperation = AnimNodeSinglePropertyHandler.Operation;
+							HandlerPtr->CopyRecords.Add(CopyRecord);
+						}
+					}
+					else
 					{
-						FExposedValueCopyRecord CopyRecord;
-						CopyRecord.DestProperty = SimpleCopyPropertyDest;
-						CopyRecord.DestArrayIndex = ArrayIt.Key();
-						CopyRecord.SourcePropertyName = AnimNodeSinglePropertyHandler.SimpleCopyPropertyName;
-						CopyRecord.SourceSubPropertyName = NAME_None;
-						CopyRecord.SourceArrayIndex = 0;
-						CopyRecord.Size = SimpleCopyPropertyDest->Inner->GetSize();
-						CopyRecord.PostCopyOperation = AnimNodeSinglePropertyHandler.Operation;
-						HandlerPtr->CopyRecords.Add(CopyRecord);
+						if (SimpleCopyPropertySource->GetSize() != SimpleCopyPropertyDest->Inner->GetSize())
+						{
+							bOnlyUsesCopyRecords = false;
+							break;
+						}
+
+						for (TMap<int32, UEdGraphPin*>::TConstIterator ArrayIt(AnimNodeSinglePropertyHandler.ArrayPins); ArrayIt; ++ArrayIt)
+						{
+							// Single variable get->array set
+							FExposedValueCopyRecord CopyRecord;
+							CopyRecord.DestProperty = SimpleCopyPropertyDest;
+							CopyRecord.DestArrayIndex = ArrayIt.Key();
+							CopyRecord.SourcePropertyName = AnimNodeSinglePropertyHandler.SimpleCopyPropertyName;
+							CopyRecord.SourceSubPropertyName = NAME_None;
+							CopyRecord.SourceArrayIndex = 0;
+							CopyRecord.Size = SimpleCopyPropertyDest->Inner->GetSize();
+							CopyRecord.PostCopyOperation = AnimNodeSinglePropertyHandler.Operation;
+							HandlerPtr->CopyRecords.Add(CopyRecord);
+						}
 					}
 				}
 				else
@@ -1938,7 +1967,7 @@ void FAnimBlueprintCompiler::FEvaluationHandlerRecord::PatchFunctionNameAndCopyR
 							break;
 						}
 
-						// Local sub-struct variable get
+						// Single sub-struct variable get->single variable set
 						FExposedValueCopyRecord CopyRecord;
 						CopyRecord.DestProperty = SimpleCopyPropertyDest;
 						CopyRecord.DestArrayIndex = 0;
@@ -1957,7 +1986,7 @@ void FAnimBlueprintCompiler::FEvaluationHandlerRecord::PatchFunctionNameAndCopyR
 							break;
 						}
 
-						// Local variable get
+						// Single variable get->single variable set
 						FExposedValueCopyRecord CopyRecord;
 						CopyRecord.DestProperty = SimpleCopyPropertyDest;
 						CopyRecord.DestArrayIndex = 0;
diff --git a/Engine/Source/Editor/LandscapeEditor/Private/LandscapeEdModeSplineTools.cpp b/Engine/Source/Editor/LandscapeEditor/Private/LandscapeEdModeSplineTools.cpp
--- a/Engine/Source/Editor/LandscapeEditor/Private/LandscapeEdModeSplineTools.cpp
+++ b/Engine/Source/Editor/LandscapeEditor/Private/LandscapeEdModeSplineTools.cpp
@@ -378,6 +378,11 @@ public:
 				AddSegment(ControlPoint, NewControlPoint, bAutoRotateOnJoin, true);
 			}
 		}
+		else
+		{
+			// required to make control point visible
+			NewControlPoint->UpdateSplinePoints();
+		}
 
 		ClearSelection();
 		SelectControlPoint(NewControlPoint);
diff --git a/Engine/Source/Editor/LandscapeEditor/Private/LandscapeSplineImportExport.cpp b/Engine/Source/Editor/LandscapeEditor/Private/LandscapeSplineImportExport.cpp
--- a/Engine/Source/Editor/LandscapeEditor/Private/LandscapeSplineImportExport.cpp
+++ b/Engine/Source/Editor/LandscapeEditor/Private/LandscapeSplineImportExport.cpp
@@ -31,7 +31,7 @@ void FLandscapeSplineTextObjectFactory::ProcessConstructedObject(UObject* Create
 	CreatedObject->PostEditImport();
 }
 
-bool FLandscapeSplineTextObjectFactory::CanCreateClass(UClass* ObjectClass) const
+bool FLandscapeSplineTextObjectFactory::CanCreateClass(UClass* ObjectClass, bool& bOmitSubObjs) const
 {
 	if (ObjectClass == ULandscapeSplineControlPoint::StaticClass() ||
 		ObjectClass == ULandscapeSplineSegment::StaticClass())
diff --git a/Engine/Source/Editor/LandscapeEditor/Private/LandscapeSplineImportExport.h b/Engine/Source/Editor/LandscapeEditor/Private/LandscapeSplineImportExport.h
--- a/Engine/Source/Editor/LandscapeEditor/Private/LandscapeSplineImportExport.h
+++ b/Engine/Source/Editor/LandscapeEditor/Private/LandscapeSplineImportExport.h
@@ -15,5 +15,5 @@ protected:
 	TArray<UObject*> OutObjects;
 
 	virtual void ProcessConstructedObject(UObject* CreatedObject) override;
-	virtual bool CanCreateClass(UClass* ObjectClass) const override;
+	virtual bool CanCreateClass(UClass* ObjectClass, bool& bOmitSubObjs) const override;
 };
