Copying //UE4/Dev-Editor to //UE4/Dev-Main (Source: //UE4/Dev-Editor @ 3228984)

diff --git a/Engine/Source/Editor/ClassViewer/Private/SClassViewer.cpp b/Engine/Source/Editor/ClassViewer/Private/SClassViewer.cpp
--- a/Engine/Source/Editor/ClassViewer/Private/SClassViewer.cpp
+++ b/Engine/Source/Editor/ClassViewer/Private/SClassViewer.cpp
@@ -73,6 +73,8 @@
 #include "Misc/TextFilterExpressionEvaluator.h"
 
 #include "SListViewSelectorDropdownMenu.h"
+#include "ClassViewerProjectSettings.h"
+#include "SScrollBorder.h"
 
 #define LOCTEXT_NAMESPACE "SClassViewer"
 
@@ -417,6 +419,13 @@ public:
 	*/
 	void UpdateClassInNode(const FString& InGeneratedClassPackageName, UClass* InNewClass, UBlueprint* InNewBluePrint );
 
+	/** Finds the node, recursively going deeper into the hierarchy. Does so by comparing class names.
+	*	@param InClassName			The name of the generated class package to find the node for.
+	*
+	*	@return The node.
+	*/
+	TSharedPtr< FClassViewerNode > FindNodeByClassName(const TSharedPtr< FClassViewerNode >& InRootNode, const FString& InClassName);
+
 private:
 	/** Recursive function to build a tree, will not filter.
 	 *	@param InOutRootNode						The node that this function will add the children of to the tree.
@@ -658,10 +667,17 @@ namespace ClassViewer
 		 *	@param bInOnlyPlaceables					Filter option to remove non-placeable classes.
 		 *	@param bInOnlyBlueprintBases				Filter option to remove non-blueprint base classes.
 		 *	@param bInShowUnloadedBlueprints			Filter option to not remove unloaded blueprints due to class filter options.
+		 *  @param InAllowedDeveloperType               Filter option for dealing with developer folders.
+		 *  @param bInInternalClasses                   Filter option for showing internal classes.
+		 *  @param InternalClasses                      The classes that have been marked as Internal Only.
+		 *  @param InternalPaths                        The paths that have been marked Internal Only.
 		 *
 		 *	@return Returns true if the child passed the filter.
 		 */		
-		static bool AddChildren_Tree(const FClassViewerInitializationOptions& InInitOptions, TSharedPtr< FClassViewerNode >& InOutRootNode, const TSharedPtr< FClassViewerNode >& InOriginalRootNode, const FTextFilterExpressionEvaluator& InTextFilter, bool bInOnlyActors, bool bInOnlyPlaceables, bool bInOnlyBlueprintBases, bool bInShowUnloadedBlueprints )
+		static bool AddChildren_Tree(const FClassViewerInitializationOptions& InInitOptions, TSharedPtr< FClassViewerNode >& InOutRootNode, 
+			const TSharedPtr< FClassViewerNode >& InOriginalRootNode, const FTextFilterExpressionEvaluator& InTextFilter, bool bInOnlyActors, 
+			bool bInOnlyPlaceables, bool bInOnlyBlueprintBases, bool bInShowUnloadedBlueprints, EClassViewerDeveloperType InAllowedDeveloperType, bool bInInternalClasses,
+			const TArray<UClass*>& InternalClasses, const TArray<FDirectoryPath>& InternalPaths)
 		{
 			if (bInOnlyActors && *InOriginalRootNode->GetClassName().Get() != FString(TEXT("Actor")))
 			{
@@ -676,6 +692,46 @@ namespace ClassViewer
 			bool bIsUnloadedBlueprint = !InOriginalRootNode->Class.IsValid();
 			bool bPassesPlaceableFilter = false;
 
+			// Determine if we allow any developer folder classes, if so determine if this class is in one of the allowed developer folders.
+			bool bPassesDeveloperFilter = true;
+			static const FString DeveloperPathWithSlash = FPackageName::FilenameToLongPackageName(FPaths::GameDevelopersDir());
+			static const FString UserDeveloperPathWithSlash = FPackageName::FilenameToLongPackageName(FPaths::GameUserDeveloperDir());
+			if (InAllowedDeveloperType == EClassViewerDeveloperType::CVDT_None)
+			{
+				bPassesDeveloperFilter = !InOriginalRootNode->GeneratedClassPackage.StartsWith(DeveloperPathWithSlash);
+			}
+			else if (InAllowedDeveloperType == EClassViewerDeveloperType::CVDT_CurrentUser)
+			{
+				if (InOriginalRootNode->GeneratedClassPackage.StartsWith(DeveloperPathWithSlash))
+					bPassesDeveloperFilter = InOriginalRootNode->GeneratedClassPackage.StartsWith(UserDeveloperPathWithSlash);
+			}
+
+			// The INI files declare classes and folders that are considered internal only. Does this class match any of those patterns?
+			// INI path: /Script/ClassViewer.ClassViewerProjectSettings
+			bool bPassesInternalFilter = true;
+			if (!bInInternalClasses && InternalPaths.Num() > 0)
+			{
+				for (int i = 0; i < InternalPaths.Num(); i++)
+				{
+					if (InOriginalRootNode->GeneratedClassPackage.StartsWith(InternalPaths[i].Path))
+					{
+						bPassesInternalFilter = false;
+						break;
+					}
+				}
+			}
+			if (!bInInternalClasses && InternalClasses.Num() > 0 && bPassesInternalFilter && InOriginalRootNode->Class.IsValid())
+			{
+				for (int i = 0; i < InternalClasses.Num(); i++)
+				{
+					if (InOriginalRootNode->Class->IsChildOf(InternalClasses[i]))
+					{
+						bPassesInternalFilter = false;
+						break;
+					}
+				}
+			}
+
 			// When in picker mode, brushes are valid "placeable" actors
 			if (bInOnlyPlaceables && InInitOptions.Mode == EClassViewerMode::ClassPicker && IsBrush(InOriginalRootNode->Class.Get()) && IsPlaceable(InOriginalRootNode->Class.Get()))
 			{
@@ -691,21 +747,22 @@ namespace ClassViewer
 			{
 				if(bInShowUnloadedBlueprints)
 				{
-					bReturnPassesFilter = InOutRootNode->bPassesFilter = bPassesPlaceableFilter && bPassesBlueprintBaseFilter && IsClassAllowed_UnloadedBlueprint(InInitOptions, InOriginalRootNode) && PassesFilter(*InOriginalRootNode->GetClassName().Get(), InTextFilter);
+					bReturnPassesFilter = InOutRootNode->bPassesFilter = bPassesPlaceableFilter && bPassesBlueprintBaseFilter && bPassesDeveloperFilter && bPassesInternalFilter && IsClassAllowed_UnloadedBlueprint(InInitOptions, InOriginalRootNode) && PassesFilter(*InOriginalRootNode->GetClassName().Get(), InTextFilter);
 				}
 			}
 			else
 			{
-				bReturnPassesFilter = InOutRootNode->bPassesFilter = bPassesPlaceableFilter && bPassesBlueprintBaseFilter && IsClassAllowed(InInitOptions, InOriginalRootNode->Class) && PassesFilter(*InOriginalRootNode->GetClassName().Get(), InTextFilter);
+				bReturnPassesFilter = InOutRootNode->bPassesFilter = bPassesPlaceableFilter && bPassesBlueprintBaseFilter && bPassesDeveloperFilter && bPassesInternalFilter && IsClassAllowed(InInitOptions, InOriginalRootNode->Class) && PassesFilter(*InOriginalRootNode->GetClassName().Get(), InTextFilter);
 			}
 
+
 			TArray< TSharedPtr< FClassViewerNode > >& ChildList = InOriginalRootNode->GetChildrenList();
 			for(int32 ChildIdx = 0; ChildIdx < ChildList.Num(); ChildIdx++)
 			{
 				TSharedPtr< FClassViewerNode > NewNode = MakeShareable( new FClassViewerNode( *ChildList[ChildIdx].Get() ) );
 
 				bReturnPassesFilter |= bChildrenPassesFilter = AddChildren_Tree(InInitOptions, NewNode, ChildList[ChildIdx], InTextFilter, false, /* bInOnlyActors - false so that anything below Actor is added */
-					bInOnlyPlaceables, bInOnlyBlueprintBases, bInShowUnloadedBlueprints);
+					bInOnlyPlaceables, bInOnlyBlueprintBases, bInShowUnloadedBlueprints, InAllowedDeveloperType, bInInternalClasses, InternalClasses, InternalPaths);
 				if(bChildrenPassesFilter)
 				{
 					InOutRootNode->AddChild(NewNode);
@@ -724,11 +781,18 @@ namespace ClassViewer
 		 *	@param bInOnlyActors						Filter option to root the tree in the "Actor" class.
 		 *	@param bInOnlyBlueprintBases				Filter option to remove non-blueprint base classes.
 		 *	@param bInShowUnloadedBlueprints			Filter option to not remove unloaded blueprints due to class filter options.
-		 *
+		 *  @param InAllowedDeveloperType               Filter option for dealing with developer folders.
+		 *  @param bInInternalClasses                   Filter option for showing internal classes.
+		 *  @param InternalClasses                      The classes that have been marked as Internal Only.
+		 *  @param InternalPaths                        The paths that have been marked Internal Only.
 		 *	@return A fully built tree.
 		 */
-		static void GetClassTree(const FClassViewerInitializationOptions& InInitOptions, TSharedPtr< FClassViewerNode >& InOutRootNode, const FTextFilterExpressionEvaluator& InTextFilter, bool bInOnlyPlaceables, bool bInOnlyActors, bool bInOnlyBlueprintBases, bool bInShowUnloadedBlueprints )
+		static void GetClassTree(const FClassViewerInitializationOptions& InInitOptions, TSharedPtr< FClassViewerNode >& InOutRootNode, 
+			const FTextFilterExpressionEvaluator& InTextFilter, bool bInOnlyPlaceables, bool bInOnlyActors, bool bInOnlyBlueprintBases, 
+			bool bInShowUnloadedBlueprints,	EClassViewerDeveloperType InAllowedDeveloperType = EClassViewerDeveloperType::CVDT_All, bool bInInternalClasses = true, 
+			const TArray<UClass*>& InternalClasses = TArray<UClass*>(), const TArray<FDirectoryPath>& InternalPaths = TArray<FDirectoryPath>())
 		{
+			
 			const TSharedPtr< FClassViewerNode > ObjectClassRoot = ClassHierarchy->GetObjectRootNode();
 
 			// Duplicate the node, it will have no children.
@@ -739,13 +803,13 @@ namespace ClassViewer
 				for(int32 ClassIdx = 0; ClassIdx < ObjectClassRoot->GetChildrenList().Num(); ClassIdx++)
 				{
 					TSharedPtr<FClassViewerNode> ChildNode = MakeShareable(new FClassViewerNode(*ObjectClassRoot->GetChildrenList()[ClassIdx].Get()));
-					if (AddChildren_Tree(InInitOptions, ChildNode, ObjectClassRoot->GetChildrenList()[ClassIdx], InTextFilter, true, bInOnlyPlaceables, bInOnlyBlueprintBases, bInShowUnloadedBlueprints))
+					if (AddChildren_Tree(InInitOptions, ChildNode, ObjectClassRoot->GetChildrenList()[ClassIdx], InTextFilter, true, bInOnlyPlaceables, bInOnlyBlueprintBases, bInShowUnloadedBlueprints, InAllowedDeveloperType, bInInternalClasses, InternalClasses, InternalPaths))
 						InOutRootNode->AddChild(ChildNode);
 				}
 			}
 			else
 			{
-				AddChildren_Tree(InInitOptions, InOutRootNode, ObjectClassRoot, InTextFilter, false, bInOnlyPlaceables, bInOnlyBlueprintBases, bInShowUnloadedBlueprints);				
+				AddChildren_Tree(InInitOptions, InOutRootNode, ObjectClassRoot, InTextFilter, false, bInOnlyPlaceables, bInOnlyBlueprintBases, bInShowUnloadedBlueprints, InAllowedDeveloperType, bInInternalClasses, InternalClasses, InternalPaths);
 			}
 		}
 
@@ -758,10 +822,18 @@ namespace ClassViewer
 		 *	@param bInOnlyPlaceables					Filter option to remove non-placeable classes.
 		 *	@param bInOnlyBlueprintBases				Filter option to remove non-blueprint base classes.
 		 *	@param bInShowUnloadedBlueprints			Filter option to not remove unloaded blueprints due to class filter options.
+		 *  @param InAllowedDeveloperType               Filter option for dealing with developer folders.
+		 *  @param bInInternalClasses                   Filter option for showing internal classes.
+		 *  @param InternalClasses                      The classes that have been marked as Internal Only.
+		 *  @param InternalPaths                        The paths that have been marked Internal Only.
 		 *
 		 *	@return Returns true if the child passed the filter.
 		 */
-		static void AddChildren_List( const FClassViewerInitializationOptions& InInitOptions, TArray< TSharedPtr< FClassViewerNode > >& InOutNodeList, const TSharedPtr< FClassViewerNode >& InOriginalRootNode, const FTextFilterExpressionEvaluator& InTextFilter, bool bInOnlyActors, bool bInOnlyPlaceables, bool bInOnlyBlueprintBases, bool bInShowUnloadedBlueprints )
+		static void AddChildren_List(const FClassViewerInitializationOptions& InInitOptions, TArray< TSharedPtr< FClassViewerNode > >& InOutNodeList,
+			const TSharedPtr< FClassViewerNode >& InOriginalRootNode, const FTextFilterExpressionEvaluator& InTextFilter,
+			bool bInOnlyActors, bool bInOnlyPlaceables, bool bInOnlyBlueprintBases, bool bInShowUnloadedBlueprints,
+			EClassViewerDeveloperType InAllowedDeveloperType, bool bInInternalClasses,
+			const TArray<UClass*>& InternalClasses, const TArray<FDirectoryPath>& InternalPaths)
 		{
 			if (bInOnlyActors && *InOriginalRootNode->GetClassName().Get() != FString(TEXT("Actor")))
 			{
@@ -772,6 +844,46 @@ namespace ClassViewer
 			bool bIsUnloadedBlueprint = !InOriginalRootNode->Class.IsValid();
 			bool bPassesPlaceableFilter = false;
 
+			// Determine if we allow any developer folder classes, if so determine if this class is in one of the allowed developer folders.
+			bool bPassesDeveloperFilter = true;
+			static const FString DeveloperPathWithSlash = FPackageName::FilenameToLongPackageName(FPaths::GameDevelopersDir());
+			static const FString UserDeveloperPathWithSlash = FPackageName::FilenameToLongPackageName(FPaths::GameUserDeveloperDir());
+			if (InAllowedDeveloperType == EClassViewerDeveloperType::CVDT_None)
+			{
+				bPassesDeveloperFilter = !InOriginalRootNode->GeneratedClassPackage.StartsWith(DeveloperPathWithSlash);
+			}
+			else if (InAllowedDeveloperType == EClassViewerDeveloperType::CVDT_CurrentUser)
+			{
+				if (InOriginalRootNode->GeneratedClassPackage.StartsWith(DeveloperPathWithSlash))
+					bPassesDeveloperFilter = InOriginalRootNode->GeneratedClassPackage.StartsWith(UserDeveloperPathWithSlash);
+			}
+
+			// The INI files declare classes and folders that are considered internal only. Does this class match any of those patterns?
+			// INI path: /Script/ClassViewer.ClassViewerProjectSettings
+			bool bPassesInternalFilter = true;
+			if (!bInInternalClasses && InternalPaths.Num() > 0)
+			{
+				for (int i = 0; i < InternalPaths.Num(); i++)
+				{
+					if (InOriginalRootNode->GeneratedClassPackage.StartsWith(InternalPaths[i].Path))
+					{
+						bPassesInternalFilter = false;
+						break;
+					}
+				}
+			}
+			if (!bInInternalClasses && InternalClasses.Num() > 0 && bPassesInternalFilter)
+			{
+				for (int i = 0; i < InternalClasses.Num(); i++)
+				{
+					if (InOriginalRootNode->Class->IsChildOf(InternalClasses[i]))
+					{
+						bPassesInternalFilter = false;
+						break;
+					}
+				}
+			}
+
 			// When in picker mode, brushes are valid "placeable" actors
 			if( bInOnlyPlaceables && InInitOptions.Mode == EClassViewerMode::ClassPicker && IsBrush(InOriginalRootNode->Class.Get()) && IsPlaceable(InOriginalRootNode->Class.Get()))
 			{
@@ -789,12 +901,12 @@ namespace ClassViewer
 			{
 				if(bInShowUnloadedBlueprints)
 				{
-					NewNode->bPassesFilter = bPassesPlaceableFilter && bPassesBlueprintBaseFilter && IsClassAllowed_UnloadedBlueprint(InInitOptions, InOriginalRootNode) && PassesFilter(*InOriginalRootNode->GetClassName().Get(), InTextFilter);
+					NewNode->bPassesFilter = bPassesPlaceableFilter && bPassesBlueprintBaseFilter && bPassesDeveloperFilter && bPassesInternalFilter && IsClassAllowed_UnloadedBlueprint(InInitOptions, InOriginalRootNode) && PassesFilter(*InOriginalRootNode->GetClassName().Get(), InTextFilter);
 				}
 			}
 			else
 			{
-				NewNode->bPassesFilter = bPassesPlaceableFilter && bPassesBlueprintBaseFilter && IsClassAllowed(InInitOptions, InOriginalRootNode->Class) && PassesFilter(*InOriginalRootNode->GetClassName().Get(), InTextFilter);
+				NewNode->bPassesFilter = bPassesPlaceableFilter && bPassesBlueprintBaseFilter && bPassesDeveloperFilter && bPassesInternalFilter &&  IsClassAllowed(InInitOptions, InOriginalRootNode->Class) && PassesFilter(*InOriginalRootNode->GetClassName().Get(), InTextFilter);
 			}
 
 			if(NewNode->bPassesFilter)
@@ -808,7 +920,7 @@ namespace ClassViewer
 			for(int32 ChildIdx = 0; ChildIdx < ChildList.Num(); ChildIdx++)
 			{
 				AddChildren_List(InInitOptions, InOutNodeList, ChildList[ChildIdx], InTextFilter, false, /* bInOnlyActors - false so that anything below Actor is added */
-					bInOnlyPlaceables, bInOnlyBlueprintBases, bInShowUnloadedBlueprints);
+					bInOnlyPlaceables, bInOnlyBlueprintBases, bInShowUnloadedBlueprints, InAllowedDeveloperType, bInInternalClasses, InternalClasses, InternalPaths);
 			}
 		}	
 
@@ -820,10 +932,17 @@ namespace ClassViewer
 		 *	@param bInOnlyActors						Filter option to root the tree in the "Actor" class.
 		 *	@param bInOnlyBlueprintBases				Filter option to remove non-blueprint base classes.
 		 *	@param bInShowUnloadedBlueprints			Filter option to not remove unloaded blueprints due to class filter options.
+		 *  @param InAllowedDeveloperType               Filter option for dealing with developer folders.
+		 *  @param bInInternalClasses                   Filter option for showing internal classes.
+		 *  @param InternalClasses                      The classes that have been marked as Internal Only.
+		 *  @param InternalPaths                        The paths that have been marked Internal Only.
 		 *
 		 *	@return A fully built list.
 		 */
-		static void GetClassList(const FClassViewerInitializationOptions& InInitOptions, TArray< TSharedPtr< FClassViewerNode > >& InOutNodeList, const FTextFilterExpressionEvaluator& InTextFilter, bool bInOnlyPlaceables, bool bInOnlyActors, bool bInOnlyBlueprintBases, bool bInShowUnloadedBlueprints)
+		static void GetClassList(const FClassViewerInitializationOptions& InInitOptions, TArray< TSharedPtr< FClassViewerNode > >& InOutNodeList, 
+			const FTextFilterExpressionEvaluator& InTextFilter, bool bInOnlyPlaceables, bool bInOnlyActors, bool bInOnlyBlueprintBases, bool bInShowUnloadedBlueprints,
+			EClassViewerDeveloperType InAllowedDeveloperType = EClassViewerDeveloperType::CVDT_All, bool bInInternalClasses = true,
+			const TArray<UClass*>& InternalClasses = TArray<UClass*>(), const TArray<FDirectoryPath>& InternalPaths = TArray<FDirectoryPath>())
 		{
 			const TSharedPtr< FClassViewerNode > ObjectClassRoot = ClassHierarchy->GetObjectRootNode();
 
@@ -844,7 +963,7 @@ namespace ClassViewer
 			TArray< TSharedPtr< FClassViewerNode > >& ChildList = ObjectClassRoot->GetChildrenList();
 			for(int32 ObjectChildIndex = 0; ObjectChildIndex < ChildList.Num(); ObjectChildIndex++)
 			{
-				AddChildren_List(InInitOptions, InOutNodeList, ChildList[ObjectChildIndex], InTextFilter, bInOnlyActors, bInOnlyPlaceables, bInOnlyBlueprintBases, bInShowUnloadedBlueprints);			
+				AddChildren_List(InInitOptions, InOutNodeList, ChildList[ObjectChildIndex], InTextFilter, bInOnlyActors, bInOnlyPlaceables, bInOnlyBlueprintBases, bInShowUnloadedBlueprints, InAllowedDeveloperType, bInInternalClasses, InternalClasses, InternalPaths);
 			}
 		}
 
@@ -1298,7 +1417,7 @@ public:
 		bDynamicClassLoading = InArgs._bDynamicClassLoading;
 		AssociatedNode = InArgs._AssociatedNode;
 		OnDoubleClicked = InArgs._OnClassItemDoubleClicked;
-
+		
 		bool bIsBlueprint(false);
 		bool bHasBlueprint(false);
 
@@ -1323,7 +1442,6 @@ public:
 				{
 					UPackage*  Package  = Class->GetOutermost();
 					UMetaData* MetaData = Package->GetMetaData();
-
 					ToolTip = FEditorClassUtils::GetTooltip(Class);
 				}
 
@@ -1666,6 +1784,31 @@ TSharedPtr< FClassViewerNode > FClassHierarchy::FindParent(const TSharedPtr< FCl
 	return ReturnNode;
 }
 
+TSharedPtr< FClassViewerNode > FClassHierarchy::FindNodeByClassName(const TSharedPtr< FClassViewerNode >& InRootNode, const FString& InClassName)
+{
+	FString NodeClassName = InRootNode->Class.IsValid() ? InRootNode->Class->GetPathName() : FString();
+	if (NodeClassName == InClassName)
+	{
+		return InRootNode;
+	}
+
+	TSharedPtr< FClassViewerNode > ReturnNode;
+
+	// Search the children recursively, one of them might have the parent.
+	for (int32 ChildClassIndex = 0; !ReturnNode.IsValid() && ChildClassIndex < InRootNode->GetChildrenList().Num(); ChildClassIndex++)
+	{
+		// Check the child, then check the return to see if it is valid. If it is valid, end the recursion.
+		ReturnNode = FindNodeByClassName(InRootNode->GetChildrenList()[ChildClassIndex], InClassName);
+
+		if (ReturnNode.IsValid())
+		{
+			break;
+		}
+	}
+
+	return ReturnNode;
+}
+
 TSharedPtr< FClassViewerNode > FClassHierarchy::FindNodeByGeneratedClassPackageName(const TSharedPtr< FClassViewerNode >& InRootNode, const FString& InGeneratedClassPackageName)
 {
 	if(InRootNode->GeneratedClassPackage == InGeneratedClassPackageName)
@@ -1939,6 +2082,12 @@ BEGIN_SLATE_FUNCTION_BUILD_OPTIMIZATION
 void SClassViewer::Construct(const FArguments& InArgs, const FClassViewerInitializationOptions& InInitOptions )
 {
 	bNeedsRefresh = true;
+	NumClasses = 0;
+	
+	bCanShowInternalClasses = true;
+	
+	// Listen for when view settings are changed
+	UClassViewerSettings::OnSettingChanged().AddSP(this, &SClassViewer::HandleSettingChanged);
 
 	InitOptions = InInitOptions;
 
@@ -1999,6 +2148,51 @@ void SClassViewer::Construct(const FArguments& InArgs, const FClassViewerInitial
 		OnContextMenuOpening = FOnContextMenuOpening::CreateSP(this, &SClassViewer::BuildMenuWidget);
 	}
 
+	SAssignNew(ClassList, SListView<TSharedPtr< FClassViewerNode > >)
+		.SelectionMode(ESelectionMode::Single)
+		.ListItemsSource(&RootTreeItems)
+		// Generates the actual widget for a tree item
+		.OnGenerateRow(this, &SClassViewer::OnGenerateRowForClassViewer)
+		// Generates the right click menu.
+		.OnContextMenuOpening(OnContextMenuOpening)
+		// Find out when the user selects something in the tree
+		.OnSelectionChanged(this, &SClassViewer::OnClassViewerSelectionChanged)
+		// Allow for some spacing between items with a larger item height.
+		.ItemHeight(20.0f)
+		.HeaderRow
+		(
+			SNew(SHeaderRow)
+			.Visibility(EVisibility::Collapsed)
+			+ SHeaderRow::Column(TEXT("Class"))
+			.DefaultLabel(NSLOCTEXT("ClassViewer", "Class", "Class"))
+		);
+	SAssignNew(ClassTree, STreeView<TSharedPtr< FClassViewerNode > >)
+		.SelectionMode(ESelectionMode::Single)
+		.TreeItemsSource(&RootTreeItems)
+		// Called to child items for any given parent item
+		.OnGetChildren(this, &SClassViewer::OnGetChildrenForClassViewerTree)
+		// Called to handle recursively expanding/collapsing items
+		.OnSetExpansionRecursive(this, &SClassViewer::SetAllExpansionStates_Helper)
+		// Generates the actual widget for a tree item
+		.OnGenerateRow(this, &SClassViewer::OnGenerateRowForClassViewer)
+		// Generates the right click menu.
+		.OnContextMenuOpening(OnContextMenuOpening)
+		// Find out when the user selects something in the tree
+		.OnSelectionChanged(this, &SClassViewer::OnClassViewerSelectionChanged)
+		// Called when the expansion state of an item changes
+		.OnExpansionChanged(this, &SClassViewer::OnClassViewerExpansionChanged)
+		// Allow for some spacing between items with a larger item height.
+		.ItemHeight(20.0f)
+		.HeaderRow
+		(
+			SNew(SHeaderRow)
+			.Visibility(EVisibility::Collapsed)
+			+ SHeaderRow::Column(TEXT("Class"))
+			.DefaultLabel(NSLOCTEXT("ClassViewer", "Class", "Class"))
+		);
+	TSharedRef<STreeView<TSharedPtr< FClassViewerNode > > > ClassTreeView = ClassTree.ToSharedRef();
+	TSharedRef<SListView<TSharedPtr< FClassViewerNode > > > ClassListView = ClassList.ToSharedRef();
+
 	// Holds the bulk of the class viewer's sub-widgets, to be added to the widget after construction
 	TSharedPtr< SWidget > ClassViewerContent;
 
@@ -2053,72 +2247,21 @@ void SClassViewer::Construct(const FArguments& InArgs, const FClassViewerInitial
 			+SVerticalBox::Slot()
 			.FillHeight(1.0f)
 			[
-				SAssignNew(ClassTree, STreeView<TSharedPtr< FClassViewerNode > >)
+				SNew(SScrollBorder, ClassTreeView)
 				.Visibility(InitOptions.DisplayMode == EClassViewerDisplayMode::TreeView ? EVisibility::Visible : EVisibility::Collapsed)
-	
-				.SelectionMode(ESelectionMode::Single)
-
-				.TreeItemsSource( &RootTreeItems )
-				// Called to child items for any given parent item
-				.OnGetChildren( this, &SClassViewer::OnGetChildrenForClassViewerTree )
-
-				// Called to handle recursively expanding/collapsing items
-				.OnSetExpansionRecursive(this, &SClassViewer::SetAllExpansionStates_Helper )
-
-				// Generates the actual widget for a tree item
-				.OnGenerateRow( this, &SClassViewer::OnGenerateRowForClassViewer ) 
-
-				// Generates the right click menu.
-				.OnContextMenuOpening( OnContextMenuOpening )
-
-				// Find out when the user selects something in the tree
-				.OnSelectionChanged( this, &SClassViewer::OnClassViewerSelectionChanged )
-
-				// Called when the expansion state of an item changes
-				.OnExpansionChanged( this, &SClassViewer::OnClassViewerExpansionChanged )
-
-				// Allow for some spacing between items with a larger item height.
-				.ItemHeight(20.0f)
-
-				.HeaderRow
-				(
-					SNew(SHeaderRow)
-					.Visibility(EVisibility::Collapsed)
-					+ SHeaderRow::Column(TEXT("Class"))
-					.DefaultLabel(NSLOCTEXT("ClassViewer", "Class", "Class"))
-				)
+				[
+					ClassTreeView
+				]
 			]
 
 			+SVerticalBox::Slot()
 			.FillHeight(1.0f)
 			[
-				SAssignNew(ClassList, SListView<TSharedPtr< FClassViewerNode > >)
-					.Visibility(InitOptions.DisplayMode == EClassViewerDisplayMode::ListView ? EVisibility::Visible : EVisibility::Collapsed)
-
-					.SelectionMode(ESelectionMode::Single)
-
-					.ListItemsSource( &RootTreeItems )
-
-					// Generates the actual widget for a tree item
-					.OnGenerateRow( this, &SClassViewer::OnGenerateRowForClassViewer ) 
-
-					// Generates the right click menu.
-					.OnContextMenuOpening( OnContextMenuOpening )
-
-					// Find out when the user selects something in the tree
-					.OnSelectionChanged( this, &SClassViewer::OnClassViewerSelectionChanged )
-
-					// Allow for some spacing between items with a larger item height.
-					.ItemHeight(20.0f)
-
-					.HeaderRow
-					(
-						SNew(SHeaderRow)
-						.Visibility(EVisibility::Collapsed)
-						+ SHeaderRow::Column(TEXT("Class"))
-						.DefaultLabel(NSLOCTEXT("ClassViewer", "Class", "Class"))
-					)
-
+				SNew(SScrollBorder, ClassListView)
+				.Visibility(InitOptions.DisplayMode == EClassViewerDisplayMode::ListView ? EVisibility::Visible : EVisibility::Collapsed)
+				[
+					ClassListView
+				]
 			]
 		]
 
@@ -2130,8 +2273,59 @@ void SClassViewer::Construct(const FArguments& InArgs, const FClassViewerInitial
 			// Asset discovery indicator
 			AssetDiscoveryIndicator
 		]
+	]
+
+	// Bottom panel
+	+ SVerticalBox::Slot()
+	.AutoHeight()
+	[
+		SNew(SHorizontalBox)
+
+		// Asset count
+		+ SHorizontalBox::Slot()
+		.FillWidth(1.f)
+		.VAlign(VAlign_Center)
+		.Padding(8, 0)
+		[
+			SNew(STextBlock)
+			.Text(this, &SClassViewer::GetClassCountText)
+		]
+
+		// View mode combo button
+		+ SHorizontalBox::Slot()
+		.AutoWidth()
+		[
+			SAssignNew(ViewOptionsComboButton, SComboButton)
+			.ContentPadding(0)
+			.ForegroundColor(this, &SClassViewer::GetViewButtonForegroundColor)
+			.ButtonStyle(FEditorStyle::Get(), "ToggleButton") // Use the tool bar item style for this button
+			.OnGetMenuContent(this, &SClassViewer::GetViewButtonContent)
+			.ButtonContent()
+			[
+				SNew(SHorizontalBox)
+					+ SHorizontalBox::Slot()
+				.AutoWidth()
+				.VAlign(VAlign_Center)
+				[
+					SNew(SImage).Image(FEditorStyle::GetBrush("GenericViewButton"))
+				]
+
+				+ SHorizontalBox::Slot()
+				.AutoWidth()
+				.Padding(2, 0, 0, 0)
+				.VAlign(VAlign_Center)
+				[
+					SNew(STextBlock).Text(LOCTEXT("ViewButton", "View Options"))
+				]
+			]
+		]
 	];
 
+	if (ViewOptionsComboButton.IsValid())
+	{
+		ViewOptionsComboButton->SetVisibility(InitOptions.bAllowViewOptions ? EVisibility::Visible : EVisibility::Collapsed);
+	}
+
 	// When using a class picker in list-view mode, the widget will auto-focus the search box
 	// and allow the up and down arrow keys to navigate and enter to pick without using the mouse ever
 	if ( InitOptions.Mode == EClassViewerMode::ClassPicker && InitOptions.DisplayMode == EClassViewerDisplayMode::ListView )
@@ -2177,6 +2371,9 @@ TSharedRef<SWidget> SClassViewer::GetContent()
 SClassViewer::~SClassViewer()
 {
 	ClassViewer::Helpers::PopulateClassviewerDelegate.RemoveAll(this);
+
+	// Remove the listener for when view settings are changed
+	UClassViewerSettings::OnSettingChanged().RemoveAll(this);
 }
 
 void SClassViewer::ClearSelection()
@@ -2332,6 +2529,169 @@ const TArray< TSharedPtr< FClassViewerNode > > SClassViewer::GetSelectedItems()
 	return ClassTree->GetSelectedItems();
 }
 
+
+const int SClassViewer::GetNumItems() const
+{
+	return NumClasses;
+}
+
+FSlateColor SClassViewer::GetViewButtonForegroundColor() const
+{
+	static const FName InvertedForegroundName("InvertedForeground");
+	static const FName DefaultForegroundName("DefaultForeground");
+
+	return ViewOptionsComboButton->IsHovered() ? FEditorStyle::GetSlateColor(InvertedForegroundName) : FEditorStyle::GetSlateColor(DefaultForegroundName);
+}
+
+TSharedRef<SWidget> SClassViewer::GetViewButtonContent()
+{
+	// Get all menu extenders for this context menu from the content browser module
+
+	FMenuBuilder MenuBuilder(/*bInShouldCloseWindowAfterMenuSelection=*/true, NULL, NULL, /*bCloseSelfOnly=*/ true);
+	MenuBuilder.BeginSection("Filters", LOCTEXT("ClassViewerFiltersHeading", "Class Filters"));
+	{
+		MenuBuilder.AddMenuEntry(
+			LOCTEXT("ShowInternalClassesOption", "Show Internal Classes"),
+			LOCTEXT("ShowInternalClassesOptionToolTip", "Shows internal-use only classes in the view."),
+			FSlateIcon(),
+			FUIAction(
+				FExecuteAction::CreateSP(this, &SClassViewer::ToggleShowInternalClasses),
+				FCanExecuteAction(),
+				FIsActionChecked::CreateSP(this, &SClassViewer::IsShowingInternalClasses)
+				),
+			NAME_None,
+			EUserInterfaceActionType::ToggleButton
+			);
+	}
+	MenuBuilder.EndSection();
+
+	MenuBuilder.BeginSection("DeveloperViewType", LOCTEXT("DeveloperViewTypeHeading", "Developer Folder Filter"));
+	{
+		MenuBuilder.AddMenuEntry(
+			LOCTEXT("NoneDeveloperViewOption", "None"),
+			LOCTEXT("NoneDeveloperViewOptionToolTip", "Filter classes to show no classes in developer folders."),
+			FSlateIcon(),
+			FUIAction(
+				FExecuteAction::CreateSP(this, &SClassViewer::SetCurrentDeveloperViewType, EClassViewerDeveloperType::CVDT_None),
+				FCanExecuteAction(),
+				FIsActionChecked::CreateSP(this, &SClassViewer::IsCurrentDeveloperViewType, EClassViewerDeveloperType::CVDT_None)
+				),
+			NAME_None,
+			EUserInterfaceActionType::RadioButton
+			);
+
+		MenuBuilder.AddMenuEntry(
+			LOCTEXT("CurrentUserDeveloperViewOption", "Current Developer"),
+			LOCTEXT("CurrentUserDeveloperViewOptionToolTip", "Filter classes to allow classes in the current user's development folder."),
+			FSlateIcon(),
+			FUIAction(
+				FExecuteAction::CreateSP(this, &SClassViewer::SetCurrentDeveloperViewType, EClassViewerDeveloperType::CVDT_CurrentUser),
+				FCanExecuteAction(),
+				FIsActionChecked::CreateSP(this, &SClassViewer::IsCurrentDeveloperViewType, EClassViewerDeveloperType::CVDT_CurrentUser)
+				),
+			NAME_None,
+			EUserInterfaceActionType::RadioButton
+			);
+
+		MenuBuilder.AddMenuEntry(
+			LOCTEXT("AllUsersDeveloperViewOption", "All Developers"),
+			LOCTEXT("AllUsersDeveloperViewOptionToolTip", "Filter classes to allow classes in all users' development folders."),
+			FSlateIcon(),
+			FUIAction(
+				FExecuteAction::CreateSP(this, &SClassViewer::SetCurrentDeveloperViewType, EClassViewerDeveloperType::CVDT_All),
+				FCanExecuteAction(),
+				FIsActionChecked::CreateSP(this, &SClassViewer::IsCurrentDeveloperViewType, EClassViewerDeveloperType::CVDT_All)
+				),
+			NAME_None,
+			EUserInterfaceActionType::RadioButton
+			);
+
+	}
+	MenuBuilder.EndSection();
+
+	return MenuBuilder.MakeWidget();
+}
+
+void SClassViewer::SetCurrentDeveloperViewType(EClassViewerDeveloperType NewType)
+{
+	if (ensure((int)NewType < (int)EClassViewerDeveloperType::CVDT_Max) && NewType != GetDefault<UClassViewerSettings>()->DeveloperFolderType)
+	{
+		GetMutableDefault<UClassViewerSettings>()->DeveloperFolderType = NewType;
+		GetMutableDefault<UClassViewerSettings>()->PostEditChange();
+	}
+}
+
+EClassViewerDeveloperType SClassViewer::GetCurrentDeveloperViewType() const
+{
+	if (!InitOptions.bAllowViewOptions)
+	{
+		return EClassViewerDeveloperType::CVDT_All;
+	}
+
+	return GetDefault<UClassViewerSettings>()->DeveloperFolderType;
+}
+
+bool SClassViewer::IsCurrentDeveloperViewType(EClassViewerDeveloperType ViewType) const
+{
+	return GetCurrentDeveloperViewType() == ViewType;
+}
+
+void SClassViewer::GetInternalOnlyClasses(TArray<FStringClassReference>& Classes)
+{
+	if (!InitOptions.bAllowViewOptions)
+	{
+		return;
+	}
+
+	Classes = GetDefault<UClassViewerProjectSettings>()->InternalOnlyClasses;
+}
+
+void SClassViewer::GetInternalOnlyPaths(TArray<FDirectoryPath>& Paths)
+{
+	if (!InitOptions.bAllowViewOptions)
+	{
+		return;
+	}
+	Paths = GetDefault<UClassViewerProjectSettings>()->InternalOnlyPaths;
+}
+
+
+FText SClassViewer::GetClassCountText() const
+{
+
+	const int32 NumAssets = GetNumItems();
+	const int32 NumSelectedAssets = GetSelectedItems().Num();
+
+	FText AssetCount = FText::GetEmpty();
+	AssetCount = LOCTEXT("AssetCountLabelSingular", "1 item");
+
+	if (NumSelectedAssets == 0)
+	{
+		if (NumAssets == 1)
+		{
+			AssetCount = LOCTEXT("AssetCountLabelSingular", "1 item");
+		}
+		else
+		{
+			AssetCount = FText::Format(LOCTEXT("AssetCountLabelPlural", "{0} items"), FText::AsNumber(NumAssets));
+		}
+	}
+	else
+	{
+		if (NumAssets == 1)
+		{
+			AssetCount = FText::Format(LOCTEXT("AssetCountLabelSingularPlusSelection", "1 item ({0} selected)"), FText::AsNumber(NumSelectedAssets));
+		}
+		else
+		{
+			AssetCount = FText::Format(LOCTEXT("AssetCountLabelPluralPlusSelection", "{0} items ({1} selected)"), FText::AsNumber(NumAssets), FText::AsNumber(NumSelectedAssets));
+		}
+	}
+
+	return AssetCount;
+}
+
+
 void SClassViewer::ExpandRootNodes()
 {
 	for (int32 NodeIdx = 0; NodeIdx < RootTreeItems.Num(); ++NodeIdx)
@@ -2582,12 +2942,25 @@ void SClassViewer::SetExpansionStatesInTree( TSharedPtr<FClassViewerNode> InItem
 	}
 }
 
+int32 SClassViewer::CountTreeItems(FClassViewerNode* Node)
+{
+	if (Node == nullptr)
+		return 0;
+	int32 Count = 1;
+	TArray<TSharedPtr<FClassViewerNode>>& ChildArray = Node->GetChildrenList();
+	for (int i = 0; i < ChildArray.Num(); i++)
+	{
+		Count += CountTreeItems(ChildArray[i].Get());
+	}
+	return Count;
+}
+
 void SClassViewer::Populate()
 {
 	bPendingSetExpansionStates = false;
 
 	// If showing a class tree, we may need to save expansion states.
-	if( ClassTree->GetVisibility() == EVisibility::Visible )
+	if(InitOptions.DisplayMode == EClassViewerDisplayMode::TreeView)
 	{
 		if( bSaveExpansionStates )
 		{
@@ -2612,14 +2985,39 @@ void SClassViewer::Populate()
 	// Empty the tree out so it can be redone.
 	RootTreeItems.Empty();
 
+	bool ShowingInternalClasses = IsShowingInternalClasses();
+
+	TArray<FStringClassReference> InternalClassNames;
+	TArray<UClass*> InternalClasses;
+	TArray<FDirectoryPath> InternalPaths;
+	// If we aren't showing the internal classes, then we need to know what classes to consider Internal Only, so let's gather them up from the settings object.
+	if (!ShowingInternalClasses)
+	{
+		GetInternalOnlyPaths(InternalPaths);
+		GetInternalOnlyClasses(InternalClassNames);
+
+		// Take the package names for the internal only classes and convert them into their UClass
+		for (int i = 0; i < InternalClassNames.Num(); i++)
+		{
+			FString PackageClassName = InternalClassNames[i].ToString();
+			const TSharedPtr<FClassViewerNode> ClassNode = ClassViewer::Helpers::ClassHierarchy->FindNodeByClassName(ClassViewer::Helpers::ClassHierarchy->GetObjectRootNode(), PackageClassName);
+
+			if (ClassNode.IsValid())
+			{
+				InternalClasses.Add(ClassNode->Class.Get());
+			}
+		}
+	}
+
 	// Based on if the list or tree is visible we create what will be displayed differently.
-	if(ClassTree->GetVisibility() == EVisibility::Visible)
+	if(InitOptions.DisplayMode == EClassViewerDisplayMode::TreeView)
 	{
 		// The root node for the tree, will be "Object" which we will skip.
 		TSharedPtr<FClassViewerNode> RootNode;
 
 		// Get the class tree, passing in certain filter options.
-		ClassViewer::Helpers::GetClassTree(InitOptions, RootNode, *TextFilterPtr, MenuPlaceableOnly_IsChecked(), MenuActorsOnly_IsChecked(), MenuBlueprintBasesOnly_IsChecked(), bShowUnloadedBlueprints );
+		ClassViewer::Helpers::GetClassTree(InitOptions, RootNode, *TextFilterPtr, MenuPlaceableOnly_IsChecked(), MenuActorsOnly_IsChecked(), MenuBlueprintBasesOnly_IsChecked(),
+			bShowUnloadedBlueprints, GetCurrentDeveloperViewType(), ShowingInternalClasses, InternalClasses, InternalPaths);
 
 		// Check if we will restore expansion states, we will not if there is filtering happening.
 		const bool bRestoreExpansionState = TextFilterPtr->GetFilterType() == ETextFilterExpressionType::Empty;
@@ -2663,16 +3061,21 @@ void SClassViewer::Populate()
 		{
 			// @todo - It would seem smart to add this in before the other items, since it needs to be on top. However, that causes strange issues with saving/restoring expansion states. 
 			// This is likely not very efficient since the list can have hundreds and even thousands of items.
-			RootTreeItems.Insert(CreateNoneOption(), 0);
+			RootTreeItems.Insert(CreateNoneOption(), 0);			
 		}
 
+		NumClasses = 0;
+		for (int i = 0; i < RootTreeItems.Num(); i++)
+			NumClasses += CountTreeItems(RootTreeItems[i].Get());
+
 		// Now that new items are in the tree, we need to request a refresh.
 		ClassTree->RequestTreeRefresh();
 	}
 	else
 	{
 		// Get the class list, passing in certain filter options.
-		ClassViewer::Helpers::GetClassList(InitOptions, RootTreeItems, *TextFilterPtr, MenuPlaceableOnly_IsChecked(), MenuActorsOnly_IsChecked(), MenuBlueprintBasesOnly_IsChecked(), bShowUnloadedBlueprints);
+		ClassViewer::Helpers::GetClassList(InitOptions, RootTreeItems, *TextFilterPtr, MenuPlaceableOnly_IsChecked(), MenuActorsOnly_IsChecked(), MenuBlueprintBasesOnly_IsChecked(), 
+			bShowUnloadedBlueprints, GetCurrentDeveloperViewType(), ShowingInternalClasses, InternalClasses, InternalPaths);
 
 		// Sort the list alphabetically.
 		struct FCompareFClassViewerNode
@@ -2697,6 +3100,10 @@ void SClassViewer::Populate()
 			RootTreeItems.Insert(CreateNoneOption(), 0);
 		}
 
+		NumClasses = 0;
+		for (int i = 0; i < RootTreeItems.Num(); i++)
+			NumClasses += CountTreeItems(RootTreeItems[i].Get() );
+
 		// Now that new items are in the list, we need to request a refresh.
 		ClassList->RequestListRefresh();
 	}
@@ -2786,4 +3193,35 @@ bool SClassViewer::IsClassAllowed(const UClass* InClass) const
 	return ClassViewer::Helpers::IsClassAllowed(InitOptions, InClass);
 }
 
+void SClassViewer::HandleSettingChanged(FName PropertyName)
+{
+	if ((PropertyName == "DisplayInternalClasses") ||
+		(PropertyName == "DeveloperFolderType") ||
+		(PropertyName == NAME_None))	// @todo: Needed if PostEditChange was called manually, for now
+	{
+		Refresh();
+	}
+}
+
+void SClassViewer::ToggleShowInternalClasses()
+{
+	check(IsToggleShowInternalClassesAllowed());
+	GetMutableDefault<UClassViewerSettings>()->DisplayInternalClasses = !GetDefault<UClassViewerSettings>()->DisplayInternalClasses;
+	GetMutableDefault<UClassViewerSettings>()->PostEditChange();
+}
+
+bool SClassViewer::IsToggleShowInternalClassesAllowed() const
+{
+	return bCanShowInternalClasses;
+}
+
+bool SClassViewer::IsShowingInternalClasses() const
+{
+	if (!InitOptions.bAllowViewOptions)
+	{
+		return true;
+	}
+	return IsToggleShowInternalClassesAllowed() ? GetDefault<UClassViewerSettings>()->DisplayInternalClasses : false;
+}
+
 #undef LOCTEXT_NAMESPACE
diff --git a/Engine/Source/Editor/ClassViewer/Public/ClassViewerModule.h b/Engine/Source/Editor/ClassViewer/Public/ClassViewerModule.h
--- a/Engine/Source/Editor/ClassViewer/Public/ClassViewerModule.h
+++ b/Engine/Source/Editor/ClassViewer/Public/ClassViewerModule.h
@@ -88,6 +88,9 @@ public:
 
 	/** The property this class viewer be working on. */
 	TSharedPtr<class IPropertyHandle> PropertyHandle;
+
+	/** true (the default) shows the view options at the bottom of the class picker*/
+	bool bAllowViewOptions;
 public:
 
 	/** Constructor */
@@ -103,6 +106,7 @@ public:
 		, bExpandRootNodes(true)
 		, bEnableClassDynamicLoading(true)
 		, bShowDisplayNames(false)
+		, bAllowViewOptions(true)
 	{
 	}
 };
diff --git a/Engine/Source/Editor/ClassViewer/Public/ClassViewerProjectSettings.h b/Engine/Source/Editor/ClassViewer/Public/ClassViewerProjectSettings.h
new file mode 100644
--- /dev/null
+++ b/Engine/Source/Editor/ClassViewer/Public/ClassViewerProjectSettings.h
@@ -0,0 +1,25 @@
+// Copyright 1998-2016 Epic Games, Inc. All Rights Reserved.
+
+#pragma once
+
+#include "CoreMinimal.h"
+#include "ClassViewerProjectSettings.generated.h"
+
+/**
+ * Implements the settings for the Class Viewer Project Settings
+ */
+UCLASS(config=Engine, defaultconfig)
+class CLASSVIEWER_API UClassViewerProjectSettings : public UObject
+{
+	GENERATED_UCLASS_BODY()
+
+#if WITH_EDITORONLY_DATA
+	/** The base directories to be considered Internal Only for the class picker.*/
+	UPROPERTY(EditAnywhere, config, Category = ClassVisibilityManagement, meta = (DisplayName = "List of directories to consider Internal Only.", ContentDir, LongPackageName))
+	TArray<FDirectoryPath> InternalOnlyPaths;
+
+	/** The base classes to be considered Internal Only for the class picker.*/
+	UPROPERTY(EditAnywhere, config, Category = ClassVisibilityManagement, meta = (MetaClass = "Object", DisplayName = "List of base classes to consider Internal Only.", AllowAbstract = "true", ShowTreeView, HideViewOptions))
+	TArray<FStringClassReference> InternalOnlyClasses;
+#endif
+};
diff --git a/Engine/Source/Editor/ContentBrowser/ContentBrowser.Build.cs b/Engine/Source/Editor/ContentBrowser/ContentBrowser.Build.cs
--- a/Engine/Source/Editor/ContentBrowser/ContentBrowser.Build.cs
+++ b/Engine/Source/Editor/ContentBrowser/ContentBrowser.Build.cs
@@ -14,6 +14,7 @@ public class ContentBrowser : ModuleRules
 				"EditorWidgets",
 				"GameProjectGeneration",
                 "MainFrame",
+				"PackagesDialog",
 				"SourceControl",
 				"SourceControlWindows",
                 "ReferenceViewer",
@@ -44,6 +45,7 @@ public class ContentBrowser : ModuleRules
 		DynamicallyLoadedModuleNames.AddRange(
 			new string[] {
 				"PropertyEditor",
+				"PackagesDialog",
 				"AssetRegistry",
 				"AssetTools",
 				"CollectionManager",
diff --git a/Engine/Source/Editor/ContentBrowser/Private/AssetContextMenu.cpp b/Engine/Source/Editor/ContentBrowser/Private/AssetContextMenu.cpp
--- a/Engine/Source/Editor/ContentBrowser/Private/AssetContextMenu.cpp
+++ b/Engine/Source/Editor/ContentBrowser/Private/AssetContextMenu.cpp
@@ -450,6 +450,25 @@ void FAssetContextMenu::MakeAssetActionsSubMenu(FMenuBuilder& MenuBuilder)
 	}
 	MenuBuilder.EndSection();
 
+	if (GetDefault<UEditorExperimentalSettings>()->bEnableContentHotReloading)
+	{
+		// EXPERIMENTAL ACTIONS
+		MenuBuilder.BeginSection("AssetContextExperimental", LOCTEXT("AssetContextAdvancedExperimentalMenuHeading", "Experimental"));
+		{
+			// Reload
+			MenuBuilder.AddMenuEntry(
+				LOCTEXT("Reload", "Reload"),
+				LOCTEXT("ReloadTooltip", "Reload the selected assets."),
+				FSlateIcon(),
+				FUIAction(
+					FExecuteAction::CreateSP(this, &FAssetContextMenu::ExecuteReload),
+					FCanExecuteAction::CreateSP(this, &FAssetContextMenu::CanExecuteReload)
+				)
+			);
+		}
+		MenuBuilder.EndSection();
+	}
+
 	// ADVANCED ACTIONS
 	MenuBuilder.BeginSection("AssetContextAdvancedActions", LOCTEXT("AssetContextAdvancedActionsMenuHeading", "Advanced"));
 	{
@@ -1829,6 +1848,72 @@ void FAssetContextMenu::ExecuteDelete()
 	}
 }
 
+bool FAssetContextMenu::CanExecuteReload() const
+{
+	TArray< FAssetData > AssetViewSelectedAssets = AssetView.Pin()->GetSelectedAssets();
+	TArray< FString > SelectedFolders = AssetView.Pin()->GetSelectedFolders();
+
+	int32 NumAssetItems, NumClassItems;
+	ContentBrowserUtils::CountItemTypes(AssetViewSelectedAssets, NumAssetItems, NumClassItems);
+
+	int32 NumAssetPaths, NumClassPaths;
+	ContentBrowserUtils::CountPathTypes(SelectedFolders, NumAssetPaths, NumClassPaths);
+
+	bool bHasSelectedCollections = false;
+	for (const FString& SelectedFolder : SelectedFolders)
+	{
+		if (ContentBrowserUtils::IsCollectionPath(SelectedFolder))
+		{
+			bHasSelectedCollections = true;
+			break;
+		}
+	}
+
+	// We can't reload classes, or folders containing classes, or any collection folders
+	return ((NumAssetItems > 0 && NumClassItems == 0) || (NumAssetPaths > 0 && NumClassPaths == 0)) && !bHasSelectedCollections;
+}
+
+void FAssetContextMenu::ExecuteReload()
+{
+	// Don't allow asset reload during PIE
+	if (GIsEditor)
+	{
+		UEditorEngine* Editor = GEditor;
+		FWorldContext* PIEWorldContext = GEditor->GetPIEWorldContext();
+		if (PIEWorldContext)
+		{
+			FNotificationInfo Notification(LOCTEXT("CannotReloadAssetInPIE", "Assets cannot be reloaded while in PIE."));
+			Notification.ExpireDuration = 3.0f;
+			FSlateNotificationManager::Get().AddNotification(Notification);
+			return;
+		}
+	}
+
+	TArray<FAssetData> AssetViewSelectedAssets = AssetView.Pin()->GetSelectedAssets();
+	if (AssetViewSelectedAssets.Num() > 0)
+	{
+		TArray<UPackage*> PackagesToReload;
+
+		for (auto AssetIt = AssetViewSelectedAssets.CreateConstIterator(); AssetIt; ++AssetIt)
+		{
+			const FAssetData& AssetData = *AssetIt;
+
+			if (AssetData.AssetClass == UObjectRedirector::StaticClass()->GetFName())
+			{
+				// Don't operate on Redirectors
+				continue;
+			}
+
+			PackagesToReload.AddUnique(AssetData.GetPackage());
+		}
+
+		if (PackagesToReload.Num() > 0)
+		{
+			PackageTools::ReloadPackages(PackagesToReload);
+		}
+	}
+}
+
 FReply FAssetContextMenu::ExecuteDeleteFolderConfirmed()
 {
 	TArray< FString > SelectedFolders = AssetView.Pin()->GetSelectedFolders();
diff --git a/Engine/Source/Editor/ContentBrowser/Private/AssetContextMenu.h b/Engine/Source/Editor/ContentBrowser/Private/AssetContextMenu.h
--- a/Engine/Source/Editor/ContentBrowser/Private/AssetContextMenu.h
+++ b/Engine/Source/Editor/ContentBrowser/Private/AssetContextMenu.h
@@ -60,6 +60,12 @@ public:
 	/** Handler for Delete */
 	void ExecuteDelete();
 
+	/** Handler to check to see if a reload command is allowed */
+	bool CanExecuteReload() const;
+
+	/** Handler for Reload */
+	void ExecuteReload();
+
 	/** Handler to check to see if "Save Asset" can be executed */
 	bool CanExecuteSaveAsset() const;
 
diff --git a/Engine/Source/Editor/ContentBrowser/Private/ContentBrowserUtils.cpp b/Engine/Source/Editor/ContentBrowser/Private/ContentBrowserUtils.cpp
--- a/Engine/Source/Editor/ContentBrowser/Private/ContentBrowserUtils.cpp
+++ b/Engine/Source/Editor/ContentBrowser/Private/ContentBrowserUtils.cpp
@@ -8,6 +8,7 @@
 #include "Misc/Paths.h"
 #include "Misc/ConfigCacheIni.h"
 #include "Misc/FeedbackContext.h"
+#include "Misc/ScopedSlowTask.h"
 #include "Misc/App.h"
 #include "Modules/ModuleManager.h"
 #include "Widgets/DeclarativeSyntaxSupport.h"
@@ -26,6 +27,7 @@
 #include "UnrealClient.h"
 #include "Engine/World.h"
 #include "Settings/ContentBrowserSettings.h"
+#include "Settings/EditorExperimentalSettings.h"
 #include "ISourceControlOperation.h"
 #include "SourceControlOperations.h"
 #include "ISourceControlModule.h"
@@ -36,6 +38,7 @@
 #include "AssetToolsModule.h"
 
 #include "Toolkits/AssetEditorManager.h"
+#include "PackagesDialog.h"
 #include "PackageTools.h"
 #include "ObjectTools.h"
 #include "ImageUtils.h"
@@ -53,6 +56,9 @@ namespace ContentBrowserUtils
 {
 	// Keep a map of all the paths that have custom colors, so updating the color in one location updates them all
 	static TMap< FString, TSharedPtr< FLinearColor > > PathColors;
+
+	/** Internal function to delete a folder from disk, but only if it is empty. InPathToDelete is in FPackageName format. */
+	bool DeleteEmptyFolderFromDisk(const FString& InPathToDelete);
 }
 
 class SContentBrowserPopup : public SCompoundWidget
@@ -347,7 +353,11 @@ bool ContentBrowserUtils::LoadAssetsIfNeeded(const TArray<FString>& ObjectPaths,
 	{
 		// Get the maximum objects to load before displaying the slow task
 		const bool bShowProgressDialog = (UnloadedObjectPaths.Num() > GetDefault<UContentBrowserSettings>()->NumObjectsToLoadBeforeWarning) || bAtLeastOneUnloadedMap;
-		GWarn->BeginSlowTask(LOCTEXT("LoadingObjects", "Loading Objects..."), bShowProgressDialog);
+		FScopedSlowTask SlowTask(UnloadedObjectPaths.Num(), LOCTEXT("LoadingObjects", "Loading Objects..."));
+		if (bShowProgressDialog)
+		{
+			SlowTask.MakeDialog();
+		}
 
 		GIsEditorLoadingPackage = true;
 
@@ -360,6 +370,7 @@ bool ContentBrowserUtils::LoadAssetsIfNeeded(const TArray<FString>& ObjectPaths,
 		for (int32 PathIdx = 0; PathIdx < UnloadedObjectPaths.Num(); ++PathIdx)
 		{
 			const FString& ObjectPath = UnloadedObjectPaths[PathIdx];
+			SlowTask.EnterProgressFrame(1, FText::Format(LOCTEXT("LoadingObjectf", "Loading {0}..."), FText::FromString(ObjectPath)));
 
 			// Load up the object
 			UObject* LoadedObject = LoadObject<UObject>(NULL, *ObjectPath, NULL, LoadFlags, NULL);
@@ -372,11 +383,6 @@ bool ContentBrowserUtils::LoadAssetsIfNeeded(const TArray<FString>& ObjectPaths,
 				bSomeObjectsFailedToLoad = true;
 			}
 
-			if ( bShowProgressDialog )
-			{
-				GWarn->UpdateProgress(PathIdx, UnloadedObjectPaths.Num());
-			}
-
 			if (GWarn->ReceivedUserCancel())
 			{
 				// If the user has canceled stop loading the remaining objects. We don't add the remaining objects to the failed string,
@@ -386,7 +392,6 @@ bool ContentBrowserUtils::LoadAssetsIfNeeded(const TArray<FString>& ObjectPaths,
 			}
 		}
 		GIsEditorLoadingPackage = false;
-		GWarn->EndSlowTask();
 
 		if ( bSomeObjectsFailedToLoad )
 		{
@@ -598,8 +603,7 @@ bool ContentBrowserUtils::DeleteFolders(const TArray<FString>& PathsToDelete)
 
 		for (const FString& PathToDelete : PathsToDelete)
 		{
-			FString PathToDeleteOnDisk;
-			if (FPackageName::TryConvertLongPackageNameToFilename(PathToDelete, PathToDeleteOnDisk) && IFileManager::Get().DeleteDirectory(*PathToDeleteOnDisk))
+			if (DeleteEmptyFolderFromDisk(PathToDelete))
 			{
 				AssetRegistryModule.Get().RemovePath(PathToDelete);
 			}
@@ -611,6 +615,45 @@ bool ContentBrowserUtils::DeleteFolders(const TArray<FString>& PathsToDelete)
 	return false;
 }
 
+bool ContentBrowserUtils::DeleteEmptyFolderFromDisk(const FString& InPathToDelete)
+{
+	struct FEmptyFolderVisitor : public IPlatformFile::FDirectoryVisitor
+	{
+		bool bIsEmpty;
+
+		FEmptyFolderVisitor()
+			: bIsEmpty(true)
+		{
+		}
+
+		virtual bool Visit(const TCHAR* FilenameOrDirectory, bool bIsDirectory) override
+		{
+			if (!bIsDirectory)
+			{
+				bIsEmpty = false;
+				return false; // abort searching
+			}
+
+			return true; // continue searching
+		}
+	};
+
+	FString PathToDeleteOnDisk;
+	if (FPackageName::TryConvertLongPackageNameToFilename(InPathToDelete, PathToDeleteOnDisk))
+	{
+		// Look for files on disk in case the folder contains things not tracked by the asset registry
+		FEmptyFolderVisitor EmptyFolderVisitor;
+		IFileManager::Get().IterateDirectoryRecursively(*PathToDeleteOnDisk, EmptyFolderVisitor);
+
+		if (EmptyFolderVisitor.bIsEmpty)
+		{
+			return IFileManager::Get().DeleteDirectory(*PathToDeleteOnDisk, false, true);
+		}
+	}
+
+	return false;
+}
+
 void ContentBrowserUtils::GetAssetsInPaths(const TArray<FString>& InPaths, TArray<FAssetData>& OutAssetDataList)
 {
 	// Load the asset registry module
@@ -811,8 +854,11 @@ bool ContentBrowserUtils::MoveFolders(const TArray<FString>& InSourcePathNames,
 			ContentBrowserUtils::MoveAssets( PathIt.Value(), Destination, PathIt.Key() );
 		}
 
-		// Attempt to remove the old paths. This operation will silently fail if any assets failed to move.
-		AssetRegistryModule.Get().RemovePath(SourcePath);
+		// Attempt to remove the old path
+		if (DeleteEmptyFolderFromDisk(SourcePath))
+		{
+			AssetRegistryModule.Get().RemovePath(SourcePath);
+		}
 	}
 
 	return true;
@@ -1095,13 +1141,13 @@ ContentBrowserUtils::ECBFolderCategory ContentBrowserUtils::GetFolderCategory( c
 	}
 	else
 	{
-		const bool bIsEngineContent = IsEngineFolder(InPath);
+		const bool bIsEngineContent = IsEngineFolder(InPath) || IsPluginFolder(InPath, EPluginLoadedFrom::Engine);
 		if(bIsEngineContent)
 		{
 			return ECBFolderCategory::EngineContent;
 		}
 
-		const bool bIsPluginContent = IsPluginFolder(InPath);
+		const bool bIsPluginContent = IsPluginFolder(InPath, EPluginLoadedFrom::GameProject);
 		if(bIsPluginContent)
 		{
 			return ECBFolderCategory::PluginContent;
@@ -1133,12 +1179,12 @@ bool ContentBrowserUtils::IsDevelopersFolder( const FString& InPath )
 	return InPath.StartsWith(DeveloperPathWithSlash) || InPath == DeveloperPathWithoutSlash;
 }
 
-bool ContentBrowserUtils::IsPluginFolder( const FString& InPath )
+bool ContentBrowserUtils::IsPluginFolder( const FString& InPath , EPluginLoadedFrom WhereFromFilter)
 {
 	FString PathWithSlash = InPath / TEXT("");
 	for(const TSharedRef<IPlugin>& Plugin: IPluginManager::Get().GetEnabledPlugins())
 	{
-		if(Plugin->CanContainContent())
+		if(Plugin->CanContainContent() && Plugin->GetLoadedFrom() == WhereFromFilter)
 		{
 			if(PathWithSlash.StartsWith(Plugin->GetMountedAssetPath()) || InPath == Plugin->GetName())
 			{
@@ -1769,65 +1815,357 @@ bool ContentBrowserUtils::IsValidPackageForCooking(const FString& PackageName, F
 	return true;
 }
 
-void ContentBrowserUtils::SyncPackagesFromSourceControl(TArray<FString> PackageNames)
+/** Given an set of packages that will be synced by a SCC operation, report any dependencies that are out-of-date and aren't in the list of packages to be synced */
+void GetOutOfDatePackageDependencies(const TArray<FString>& InPackagesThatWillBeSynced, TArray<FString>& OutDependenciesThatAreOutOfDate)
 {
-	// Form a list of loaded packages to prompt for save
-	TArray<UPackage*> LoadedPackages;
-	TArray<FString> LoadedPackageNames;
-	for (const auto& PackageName : PackageNames)
+	FAssetRegistryModule& AssetRegistryModule = FModuleManager::LoadModuleChecked<FAssetRegistryModule>(TEXT("AssetRegistry"));
+
+	// Build up the initial list of known packages
+	// We add to these as we find new dependencies to process
+	TSet<FName> AllPackages;
+	TArray<FName> AllPackagesArray;
 	{
-		UPackage* Package = FindPackage(nullptr, *PackageName);
-		if (Package != nullptr)
+		AllPackages.Reserve(InPackagesThatWillBeSynced.Num());
+		AllPackagesArray.Reserve(InPackagesThatWillBeSynced.Num());
+
+		for (const FString& PackageName : InPackagesThatWillBeSynced)
 		{
-			LoadedPackages.Add(Package);
-			LoadedPackageNames.Add(PackageName);
+			const FName PackageFName = *PackageName;
+			AllPackages.Emplace(PackageFName);
+			AllPackagesArray.Emplace(PackageFName);
 		}
 	}
 
-	// Prevent any level assets from being synced if they are currently being edited
-	bool bSyncingLevelBeingEdited = false;
-	LoadedPackages.RemoveAllSwap(
-		[&bSyncingLevelBeingEdited, &PackageNames, &LoadedPackageNames](UPackage* Package)
+	// Build up the complete set of package dependencies
+	TArray<FString> AllDependencies;
+	{
+		for (int32 PackageIndex = 0; PackageIndex < AllPackagesArray.Num(); ++PackageIndex)
 		{
-			UWorld* ExistingWorld = UWorld::FindWorldInPackage(Package);
-			if (ExistingWorld && ExistingWorld->WorldType == EWorldType::Editor)
+			const FName PackageName = AllPackagesArray[PackageIndex];
+
+			TArray<FName> PackageDependencies;
+			AssetRegistryModule.GetDependencies(PackageName, PackageDependencies, EAssetRegistryDependencyType::Packages);
+
+			for (const FName PackageDependency : PackageDependencies)
 			{
-				PackageNames.RemoveSwap(Package->GetName());
-				LoadedPackageNames.RemoveSwap(Package->GetName());
-				bSyncingLevelBeingEdited = true;
-				return true;
+				if (!AllPackages.Contains(PackageDependency))
+				{
+					AllPackages.Emplace(PackageDependency);
+					AllPackagesArray.Emplace(PackageDependency);
+
+					FString PackageDependencyStr = PackageDependency.ToString();
+					if (!FPackageName::IsScriptPackage(PackageDependencyStr))
+					{
+						AllDependencies.Emplace(MoveTemp(PackageDependencyStr));
+					}
+				}
 			}
+		}
+	}
 
-			return false;
+	// Query SCC to see which dependencies are out-of-date
+	if (AllDependencies.Num() > 0)
+	{
+		ISourceControlProvider& SCCProvider = ISourceControlModule::Get().GetProvider();
+		const TArray<FString> DependencyFilenames = SourceControlHelpers::PackageFilenames(AllDependencies);
+
+		SCCProvider.Execute(ISourceControlOperation::Create<FUpdateStatus>(), DependencyFilenames);
+		for (int32 DependencyIndex = 0; DependencyIndex < AllDependencies.Num(); ++DependencyIndex)
+		{
+			const FString& DependencyName = AllDependencies[DependencyIndex];
+			const FString& DependencyFilename = DependencyFilenames[DependencyIndex];
+
+			FSourceControlStatePtr SCCState = SCCProvider.GetState(DependencyFilename, EStateCacheUsage::Use);
+			if (SCCState.IsValid() && !SCCState->IsCurrent())
+			{
+				OutDependenciesThatAreOutOfDate.Emplace(DependencyName);
+			}
 		}
-	);
+	}
+}
 
-	if (bSyncingLevelBeingEdited)
+void ShowSyncDependenciesDialog(const TArray<FString>& InDependencies, TArray<FString>& OutExtraPackagesToSync)
+{
+	if (InDependencies.Num() > 0)
 	{
-		FNotificationInfo Info(LOCTEXT("CannotSyncLevelBeingEdited", "Level(s) being currently edited have not been synced."));
-		Info.ExpireDuration = 3.0f;
-		FSlateNotificationManager::Get().AddNotification(Info);
+		FPackagesDialogModule& PackagesDialogModule = FModuleManager::LoadModuleChecked<FPackagesDialogModule>(TEXT("PackagesDialog"));
+
+		PackagesDialogModule.CreatePackagesDialog(
+			LOCTEXT("SyncAssetDependenciesTitle", "Sync Asset Dependencies"), 
+			LOCTEXT("SyncAssetDependenciesMessage", "The following assets have newer versions available, but aren't selected to be synced.\nSelect any additional dependencies you would like to sync in order to avoid potential issues loading the updated packages.")
+			);
+
+		PackagesDialogModule.AddButton(
+			DRT_CheckOut, 
+			LOCTEXT("SyncDependenciesButton", "Sync"),
+			LOCTEXT("SyncDependenciesButtonTip", "Sync with the selected dependencies included")
+			);
+
+		for (const FString& DependencyName : InDependencies)
+		{
+			UPackage* Package = FindPackage(nullptr, *DependencyName);
+			PackagesDialogModule.AddPackageItem(Package, DependencyName, ECheckBoxState::Checked);
+		}
+
+		const EDialogReturnType UserResponse = PackagesDialogModule.ShowPackagesDialog();
+
+		if (UserResponse == DRT_CheckOut)
+		{
+			TArray<UPackage*> SelectedPackages;
+			PackagesDialogModule.GetResults(SelectedPackages, ECheckBoxState::Checked);
+
+			for (UPackage* SelectedPackage : SelectedPackages)
+			{
+				OutExtraPackagesToSync.Emplace(SelectedPackage->GetName());
+			}
+		}
 	}
+}
 
+void ContentBrowserUtils::SyncPackagesFromSourceControl(const TArray<FString>& PackageNames)
+{
 	if (PackageNames.Num() > 0)
 	{
-		FText ErrorMessage;
-		PackageTools::UnloadPackages(LoadedPackages, ErrorMessage);
+		// Warn about any packages that are being synced without also getting the newest version of their dependencies...
+		TArray<FString> PackageNamesToSync = PackageNames;
+		{
+			TArray<FString> OutOfDateDependencies;
+			GetOutOfDatePackageDependencies(PackageNamesToSync, OutOfDateDependencies);
 
-		if (!ErrorMessage.IsEmpty())
+			TArray<FString> ExtraPackagesToSync;
+			ShowSyncDependenciesDialog(OutOfDateDependencies, ExtraPackagesToSync);
+			
+			PackageNamesToSync.Append(ExtraPackagesToSync);
+		}
+
+		if (GetDefault<UEditorExperimentalSettings>()->bEnableContentHotReloading)
 		{
-			FMessageDialog::Open(EAppMsgType::Ok, ErrorMessage);
+			ISourceControlProvider& SCCProvider = ISourceControlModule::Get().GetProvider();
+			const TArray<FString> PackageFilenames = SourceControlHelpers::PackageFilenames(PackageNamesToSync);
+
+			// Form a list of loaded packages to reload...
+			TArray<UPackage*> LoadedPackages;
+			LoadedPackages.Reserve(PackageNamesToSync.Num());
+			for (const FString& PackageName : PackageNamesToSync)
+			{
+				UPackage* Package = FindPackage(nullptr, *PackageName);
+				if (Package)
+				{
+					LoadedPackages.Emplace(Package);
+
+					// Detach the linkers of any loaded packages so that SCC can overwrite the files...
+					if (!Package->IsFullyLoaded())
+					{
+						FlushAsyncLoading();
+						Package->FullyLoad();
+					}
+					ResetLoaders(Package);
+				}
+			}
+
+			// Sync everything...
+			SCCProvider.Execute(ISourceControlOperation::Create<FSync>(), PackageFilenames);
+
+			// Syncing may have deleted some packages, so we need to unload those rather than re-load them...
+			TArray<UPackage*> PackagesToUnload;
+			LoadedPackages.RemoveAll([&](UPackage* InPackage) -> bool
+			{
+				const FString PackageExtension = InPackage->ContainsMap() ? FPackageName::GetMapPackageExtension() : FPackageName::GetAssetPackageExtension();
+				const FString PackageFilename = FPackageName::LongPackageNameToFilename(InPackage->GetName(), PackageExtension);
+				if (!FPaths::FileExists(PackageFilename))
+				{
+					PackagesToUnload.Emplace(InPackage);
+					return true; // remove package
+				}
+				return false; // keep package
+			});
+
+			// Hot-reload the new packages...
+			PackageTools::ReloadPackages(LoadedPackages);
+
+			// Unload any deleted packages...
+			PackageTools::UnloadPackages(PackagesToUnload);
+
+			// Re-cache the SCC state...
+			SCCProvider.Execute(ISourceControlOperation::Create<FUpdateStatus>(), PackageFilenames, EConcurrency::Asynchronous);
 		}
 		else
+		{
+			// Form a list of loaded packages to unload
+			TArray<UPackage*> LoadedPackages;
+			TArray<FString> LoadedPackageNames;
+			for (const FString& PackageName : PackageNamesToSync)
+			{
+				UPackage* Package = FindPackage(nullptr, *PackageName);
+				if (Package != nullptr)
+				{
+					LoadedPackages.Add(Package);
+					LoadedPackageNames.Add(PackageName);
+				}
+			}
+
+			// Prevent any level assets from being synced if they are currently being edited
+			bool bSyncingLevelBeingEdited = false;
+			LoadedPackages.RemoveAllSwap([&bSyncingLevelBeingEdited, &PackageNamesToSync, &LoadedPackageNames](UPackage* Package)
+			{
+				UWorld* ExistingWorld = UWorld::FindWorldInPackage(Package);
+				if (ExistingWorld && ExistingWorld->WorldType == EWorldType::Editor)
+				{
+					PackageNamesToSync.RemoveSwap(Package->GetName());
+					LoadedPackageNames.RemoveSwap(Package->GetName());
+					bSyncingLevelBeingEdited = true;
+					return true;
+				}
+
+				return false;
+			});
+
+			if (bSyncingLevelBeingEdited)
+			{
+				FNotificationInfo Info(LOCTEXT("CannotSyncLevelBeingEdited", "Level(s) being currently edited have not been synced."));
+				Info.ExpireDuration = 3.0f;
+				FSlateNotificationManager::Get().AddNotification(Info);
+			}
+
+			FText ErrorMessage;
+			PackageTools::UnloadPackages(LoadedPackages, ErrorMessage);
+
+			if (!ErrorMessage.IsEmpty())
+			{
+				FMessageDialog::Open(EAppMsgType::Ok, ErrorMessage);
+			}
+			else
+			{
+				ISourceControlProvider& SCCProvider = ISourceControlModule::Get().GetProvider();
+				const TArray<FString> PackageFilenames = SourceControlHelpers::PackageFilenames(PackageNamesToSync);
+				SCCProvider.Execute(ISourceControlOperation::Create<FSync>(), PackageFilenames);
+				for (const FString& PackageName : LoadedPackageNames)
+				{
+					PackageTools::LoadPackage(PackageName);
+				}
+				SCCProvider.Execute(ISourceControlOperation::Create<FUpdateStatus>(), PackageFilenames, EConcurrency::Asynchronous);
+			}
+		}
+	}
+}
+
+void ContentBrowserUtils::SyncPathsFromSourceControl(const TArray<FString>& ContentPaths)
+{
+	TArray<FString> PathsOnDisk;
+	PathsOnDisk.Reserve(ContentPaths.Num());
+	for (const FString& ContentPath : ContentPaths)
+	{
+		FString PathOnDisk;
+		if (FPackageName::TryConvertLongPackageNameToFilename(ContentPath / TEXT(""), PathOnDisk) && FPaths::DirectoryExists(PathOnDisk))
+		{
+			PathsOnDisk.Emplace(MoveTemp(PathOnDisk));
+		}
+	}
+
+	if (PathsOnDisk.Num() > 0)
+	{
+		// Get all the assets under the path(s) on disk...
+		TArray<FString> PackageNames;
+		{
+			FAssetRegistryModule& AssetRegistryModule = FModuleManager::Get().LoadModuleChecked<FAssetRegistryModule>(TEXT("AssetRegistry"));
+
+			FARFilter Filter;
+			Filter.bRecursivePaths = true;
+			for (const FString& PathOnDisk : PathsOnDisk)
+			{
+				FString PackagePath = FPackageName::FilenameToLongPackageName(PathOnDisk);
+				if (PackagePath.Len() > 1 && PackagePath[PackagePath.Len() - 1] == TEXT('/'))
+				{
+					// The filter path can't end with a trailing slash
+					PackagePath = PackagePath.LeftChop(1);
+				}
+				Filter.PackagePaths.Emplace(*PackagePath);
+			}
+
+			TArray<FAssetData> AssetList;
+			AssetRegistryModule.Get().GetAssets(Filter, AssetList);
+
+			TSet<FName> UniquePackageNames;
+			for (const FAssetData& Asset : AssetList)
+			{
+				bool bWasInSet = false;
+				UniquePackageNames.Add(Asset.PackageName, &bWasInSet);
+				if (!bWasInSet)
+				{
+					PackageNames.Add(Asset.PackageName.ToString());
+				}
+			}
+		}
+
+		if (GetDefault<UEditorExperimentalSettings>()->bEnableContentHotReloading)
 		{
 			ISourceControlProvider& SCCProvider = ISourceControlModule::Get().GetProvider();
-			const TArray<FString> PackageFilenames = SourceControlHelpers::PackageFilenames(PackageNames);
-			SCCProvider.Execute(ISourceControlOperation::Create<FSync>(), PackageFilenames);
-			for (const FString& PackageName : LoadedPackageNames)
+
+			// Warn about any packages that are being synced without also getting the newest version of their dependencies...
+			TArray<FString> PackageNamesToSync = PackageNames;
+			TArray<FString> ExtraPackagesToSync;
 			{
-				PackageTools::LoadPackage(PackageName);
+				TArray<FString> OutOfDateDependencies;
+				GetOutOfDatePackageDependencies(PackageNamesToSync, OutOfDateDependencies);
+
+				ShowSyncDependenciesDialog(OutOfDateDependencies, ExtraPackagesToSync);
+
+				PackageNamesToSync.Append(ExtraPackagesToSync);
 			}
-			SCCProvider.Execute(ISourceControlOperation::Create<FUpdateStatus>(), PackageFilenames, EConcurrency::Asynchronous);
+
+			// Form a list of loaded packages to reload...
+			TArray<UPackage*> LoadedPackages;
+			LoadedPackages.Reserve(PackageNamesToSync.Num());
+			for (const FString& PackageName : PackageNamesToSync)
+			{
+				UPackage* Package = FindPackage(nullptr, *PackageName);
+				if (Package)
+				{
+					LoadedPackages.Emplace(Package);
+
+					// Detach the linkers of any loaded packages so that SCC can overwrite the files...
+					if (!Package->IsFullyLoaded())
+					{
+						FlushAsyncLoading();
+						Package->FullyLoad();
+					}
+					ResetLoaders(Package);
+				}
+			}
+
+			// Sync everything...
+			SCCProvider.Execute(ISourceControlOperation::Create<FSync>(), PathsOnDisk);
+			if (ExtraPackagesToSync.Num() > 0)
+			{
+				SCCProvider.Execute(ISourceControlOperation::Create<FSync>(), SourceControlHelpers::PackageFilenames(ExtraPackagesToSync));
+			}
+
+			// Syncing may have deleted some packages, so we need to unload those rather than re-load them...
+			TArray<UPackage*> PackagesToUnload;
+			LoadedPackages.RemoveAll([&](UPackage* InPackage) -> bool
+			{
+				const FString PackageExtension = InPackage->ContainsMap() ? FPackageName::GetMapPackageExtension() : FPackageName::GetAssetPackageExtension();
+				const FString PackageFilename = FPackageName::LongPackageNameToFilename(InPackage->GetName(), PackageExtension);
+				if (!FPaths::FileExists(PackageFilename))
+				{
+					PackagesToUnload.Emplace(InPackage);
+					return true; // remove package
+				}
+				return false; // keep package
+			});
+
+			// Hot-reload the new packages...
+			PackageTools::ReloadPackages(LoadedPackages);
+
+			// Unload any deleted packages...
+			PackageTools::UnloadPackages(PackagesToUnload);
+
+			// Re-cache the SCC state...
+			SCCProvider.Execute(ISourceControlOperation::Create<FUpdateStatus>(), PathsOnDisk, EConcurrency::Asynchronous);
+		}
+		else
+		{
+			ContentBrowserUtils::SyncPackagesFromSourceControl(PackageNames);
 		}
 	}
 }
diff --git a/Engine/Source/Editor/ContentBrowser/Private/ContentBrowserUtils.h b/Engine/Source/Editor/ContentBrowser/Private/ContentBrowserUtils.h
--- a/Engine/Source/Editor/ContentBrowser/Private/ContentBrowserUtils.h
+++ b/Engine/Source/Editor/ContentBrowser/Private/ContentBrowserUtils.h
@@ -7,6 +7,7 @@
 #include "Framework/SlateDelegates.h"
 #include "AssetData.h"
 #include "CollectionManagerTypes.h"
+#include "IPluginManager.h"
 
 class FViewport;
 
@@ -152,8 +153,8 @@ namespace ContentBrowserUtils
 	/** Returns true if the passed-in path is a developers folder */
 	bool IsDevelopersFolder( const FString& InPath );
 
-	/** Returns true if the passed-in path is a plugin folder */
-	bool IsPluginFolder( const FString& InPath );
+	/** Returns true if the passed-in path is a plugin folder matching the specified "where from" filter.*/
+	bool IsPluginFolder( const FString& InPath, EPluginLoadedFrom WhereFromFilter);
 
 	/** Returns true if the passed-in path is a C++ classes folder */
 	bool IsClassesFolder( const FString& InPath );
@@ -247,7 +248,10 @@ namespace ContentBrowserUtils
 	bool IsValidPackageForCooking(const FString& PackageName, FText& OutErrorMessage);
 
 	/** Syncs the specified packages from source control, other than any level assets which are currently being edited */
-	void SyncPackagesFromSourceControl(TArray<FString> PackageNames);
+	void SyncPackagesFromSourceControl(const TArray<FString>& PackageNames);
+
+	/** Syncs the content from the specified paths from source control, other than any level assets which are currently being edited */
+	void SyncPathsFromSourceControl(const TArray<FString>& ContentPaths);
 
 	// We assume the game name is 20 characters (the maximum allowed) to make sure that content can be ported between projects
 	// 260 characters is the limit on Windows, which is the shortest max path of any platforms that support cooking
diff --git a/Engine/Source/Editor/ContentBrowser/Private/NativeClassHierarchy.cpp b/Engine/Source/Editor/ContentBrowser/Private/NativeClassHierarchy.cpp
--- a/Engine/Source/Editor/ContentBrowser/Private/NativeClassHierarchy.cpp
+++ b/Engine/Source/Editor/ContentBrowser/Private/NativeClassHierarchy.cpp
@@ -94,10 +94,11 @@ void FNativeClassHierarchy::GetClassFolders(TArray<FName>& OutClassRoots, TArray
 
 	for(const auto& RootNode : RootNodes)
 	{
-		const bool bRootNodePassesFilter = 
+		bool bRootNodePassesFilter =
 			(RootNode.Key == GameRootNodeName) ||								// Always include game classes
-			(RootNode.Key == EngineRootNodeName && bIncludeEngineClasses) ||	// Only include engine classes if we were asked for them
-			(RootNode.Key != EngineRootNodeName && bIncludePluginClasses);		// Only include plugin classes if we were asked for them
+			(RootNode.Key == EngineRootNodeName && bIncludeEngineClasses) ||  	// Only include engine classes if we were asked for them
+			(bIncludePluginClasses && RootNode.Value->LoadedFrom == EPluginLoadedFrom::GameProject) || // Only include Game plugin classes if we were asked for them
+			(bIncludePluginClasses && bIncludeEngineClasses && RootNode.Value->LoadedFrom == EPluginLoadedFrom::Engine); //  Only include engine plugin classes if we were asked for them
 
 		if(bRootNodePassesFilter)
 		{
@@ -186,7 +187,7 @@ void FNativeClassHierarchy::PopulateHierarchy()
 	RootNodes.Empty();
 
 	TSet<FName> GameModules = GetGameModules();
-	TMap<FName, FName> PluginModules = GetPluginModules();
+	TMap<FName, FNativeClassHierarchyPluginModuleInfo> PluginModules = GetPluginModules();
 
 	for(TObjectIterator<UClass> ClassIt; ClassIt; ++ClassIt)
 	{
@@ -211,7 +212,7 @@ void FNativeClassHierarchy::AddClassesForModule(const FName& InModuleName)
 	}
 
 	TSet<FName> GameModules = GetGameModules();
-	TMap<FName, FName> PluginModules = GetPluginModules();
+	TMap<FName, FNativeClassHierarchyPluginModuleInfo> PluginModules = GetPluginModules();
 
 	TArray<UObject*> PackageObjects;
 	GetObjectsWithOuter(ClassPackage, PackageObjects, false);
@@ -254,7 +255,7 @@ void FNativeClassHierarchy::RemoveClassesForModule(const FName& InModuleName)
 	}
 }
 
-void FNativeClassHierarchy::AddClass(UClass* InClass, const TSet<FName>& InGameModules, const TMap<FName, FName>& InPluginModules, FAddClassMetrics& AddClassMetrics)
+void FNativeClassHierarchy::AddClass(UClass* InClass, const TSet<FName>& InGameModules, const TMap<FName, FNativeClassHierarchyPluginModuleInfo>& InPluginModules, FAddClassMetrics& AddClassMetrics)
 {
 	// Ignore deprecated and temporary classes
 	if(InClass->HasAnyClassFlags(CLASS_Deprecated | CLASS_NewerVersionExists) || FKismetEditorUtilities::IsClassABlueprintSkeleton(InClass))
@@ -276,7 +277,8 @@ void FNativeClassHierarchy::AddClass(UClass* InClass, const TSet<FName>& InGameM
 	}
 
 	// Work out which root this class should go under
-	const FName RootNodeName = GetClassPathRootForModule(ClassModuleName, InGameModules, InPluginModules);
+	EPluginLoadedFrom WhereLoadedFrom;
+	const FName RootNodeName = GetClassPathRootForModule(ClassModuleName, InGameModules, InPluginModules, WhereLoadedFrom);
 
 	// Work out the final path to this class within the hierarchy (which isn't the same as the path on disk)
 	const FString ClassModuleRelativePath = ClassModuleRelativeIncludePath.Left(ClassModuleRelativeIncludePath.Find(TEXT("/"), ESearchCase::CaseSensitive, ESearchDir::FromEnd));
@@ -287,6 +289,7 @@ void FNativeClassHierarchy::AddClass(UClass* InClass, const TSet<FName>& InGameM
 	if(!RootNode.IsValid())
 	{
 		RootNode = FNativeClassHierarchyNode::MakeFolderEntry(RootNodeName, TEXT("/") + RootNodeName.ToString());
+		RootNode->LoadedFrom = WhereLoadedFrom;
 		++AddClassMetrics.NumFoldersAdded;
 	}
 
@@ -392,10 +395,11 @@ bool FNativeClassHierarchy::GetClassPath(UClass* InClass, FString& OutClassPath,
 	}
 
 	TSet<FName> GameModules = GetGameModules();
-	TMap<FName, FName> PluginModules = GetPluginModules();
+	TMap<FName, FNativeClassHierarchyPluginModuleInfo> PluginModules = GetPluginModules();
 
 	// Work out which root this class should go under
-	const FName RootNodeName = GetClassPathRootForModule(ClassModuleName, GameModules, PluginModules);
+	EPluginLoadedFrom WhereLoadedFrom;
+	const FName RootNodeName = GetClassPathRootForModule(ClassModuleName, GameModules, PluginModules, WhereLoadedFrom);
 
 	// Work out the final path to this class within the hierarchy (which isn't the same as the path on disk)
 	const FString ClassModuleRelativePath = ClassModuleRelativeIncludePath.Left(ClassModuleRelativeIncludePath.Find(TEXT("/"), ESearchCase::CaseSensitive, ESearchDir::FromEnd));
@@ -448,21 +452,25 @@ FName FNativeClassHierarchy::GetClassModuleName(UClass* InClass)
 	return NAME_None;
 }
 
-FName FNativeClassHierarchy::GetClassPathRootForModule(const FName& InModuleName, const TSet<FName>& InGameModules, const TMap<FName, FName>& InPluginModules)
+FName FNativeClassHierarchy::GetClassPathRootForModule(const FName& InModuleName, const TSet<FName>& InGameModules, const TMap<FName, FNativeClassHierarchyPluginModuleInfo>& InPluginModules, EPluginLoadedFrom& OutWhereLoadedFrom)
 {
 	static const FName EngineRootNodeName = "Classes_Engine";
 	static const FName GameRootNodeName = "Classes_Game";
 
 	// Work out which root this class should go under (anything that isn't a game or plugin module goes under engine)
 	FName RootNodeName = EngineRootNodeName;
+	OutWhereLoadedFrom = EPluginLoadedFrom::Engine;
+
 	if(InGameModules.Contains(InModuleName))
 	{
 		RootNodeName = GameRootNodeName;
+		OutWhereLoadedFrom = EPluginLoadedFrom::GameProject;
 	}
 	else if(InPluginModules.Contains(InModuleName))
 	{
-		const FName PluginName = InPluginModules.FindRef(InModuleName);
-		RootNodeName = FName(*(FString(TEXT("Classes_")) + PluginName.ToString()));
+		const FNativeClassHierarchyPluginModuleInfo PluginInfo = InPluginModules.FindRef(InModuleName);
+		RootNodeName = FName(*(FString(TEXT("Classes_")) + PluginInfo.Name.ToString()));
+		OutWhereLoadedFrom = PluginInfo.LoadedFrom;
 	}
 
 	return RootNodeName;
@@ -486,19 +494,22 @@ TSet<FName> FNativeClassHierarchy::GetGameModules()
 	return GameModules;
 }
 
-TMap<FName, FName> FNativeClassHierarchy::GetPluginModules()
+TMap<FName, FNativeClassHierarchyPluginModuleInfo> FNativeClassHierarchy::GetPluginModules()
 {
 	IPluginManager& PluginManager = IPluginManager::Get();
 
 	// Build up a map of plugin modules -> plugin names - used to work out which modules populate a given Classes_PluginName
-	TMap<FName, FName> PluginModules;
+	TMap<FName, FNativeClassHierarchyPluginModuleInfo> PluginModules;
 	{
 		TArray<FPluginStatus> Plugins = PluginManager.QueryStatusForAllPlugins();
 		for(const FPluginStatus& Plugin: Plugins)
 		{
 			for(const FModuleDescriptor& PluginModule: Plugin.Descriptor.Modules)
 			{
-				PluginModules.Add(PluginModule.Name, FName(*Plugin.Name));
+				FNativeClassHierarchyPluginModuleInfo info;
+				info.Name = FName(*Plugin.Name);
+				info.LoadedFrom = Plugin.LoadedFrom;
+				PluginModules.Add(PluginModule.Name, info);
 			}
 		}
 	}
diff --git a/Engine/Source/Editor/ContentBrowser/Private/NativeClassHierarchy.h b/Engine/Source/Editor/ContentBrowser/Private/NativeClassHierarchy.h
--- a/Engine/Source/Editor/ContentBrowser/Private/NativeClassHierarchy.h
+++ b/Engine/Source/Editor/ContentBrowser/Private/NativeClassHierarchy.h
@@ -4,6 +4,7 @@
 
 #include "CoreMinimal.h"
 #include "HAL/PlatformTime.h"
+#include "IPluginManager.h"
 
 /**
  * Type of hierarchy node
@@ -14,6 +15,18 @@ enum class ENativeClassHierarchyNodeType : uint8
 	Class,
 };
 
+/**
+ *  Contains high-level information about the plugin module relevant to the native class hierarchy.                                                                      
+ */
+struct FNativeClassHierarchyPluginModuleInfo
+{
+	/** Name of the module*/
+	FName Name;
+
+	/** Indicator of where the module was loaded from (Engine or GameProject)*/
+	EPluginLoadedFrom LoadedFrom;
+};
+
 /**
  * Type used as a key in a map to resolve name conflicts between folders and classes
  */
@@ -83,6 +96,9 @@ struct FNativeClassHierarchyNode
 
 	/** Child entries (Type == Folder) */
 	TMap<FNativeClassHierarchyNodeKey, TSharedPtr<FNativeClassHierarchyNode>> Children;
+
+	/** Which type of plugin this data was originally loaded from (if loaded from a plugin)*/
+	EPluginLoadedFrom LoadedFrom;
 };
 
 /**
@@ -267,7 +283,7 @@ private:
 	 * @param InPluginModules - The list of modules that belong to neither the "/Classes_Game" or "/Classes_Engine" roots
 	 * @param AddClassMetrics - Metrics to update as new classes and folders are added (used to report population performance)
 	 */
-	void AddClass(UClass* InClass, const TSet<FName>& InGameModules, const TMap<FName, FName>& InPluginModules, FAddClassMetrics& AddClassMetrics);
+	void AddClass(UClass* InClass, const TSet<FName>& InGameModules, const TMap<FName, FNativeClassHierarchyPluginModuleInfo>& InPluginModules, FAddClassMetrics& AddClassMetrics);
 
 	/**
 	 * Called when we're notified that a module has changed status
@@ -299,10 +315,11 @@ private:
 	 * @param InModuleName - The module we want to find the root path for
 	 * @param InGameModules - The list of modules that belong to the "/Classes_Game" root
 	 * @param InPluginModules - The list of modules that belong to neither the "/Classes_Game" or "/Classes_Engine" roots
-	 *
+	 * @param WhereLoadedFrom - Determine where this class was loaded from.
+	 * 
 	 * @return The name of the root path to use, eg "/Classes_Game"
 	 */
-	static FName GetClassPathRootForModule(const FName& InModuleName, const TSet<FName>& InGameModules, const TMap<FName, FName>& InPluginModules);
+	static FName GetClassPathRootForModule(const FName& InModuleName, const TSet<FName>& InGameModules, const TMap<FName, FNativeClassHierarchyPluginModuleInfo>& InPluginModules, EPluginLoadedFrom& OutWhereLoadedFrom);
 
 	/**
 	 * Calculate the list of modules that belong to the "/Classes_Game" root
@@ -312,7 +329,7 @@ private:
 	/**
 	 * Calculate the list of modules that belong to neither the "/Classes_Game" or "/Classes_Engine" roots
 	 */
-	static TMap<FName, FName> GetPluginModules();
+	static TMap<FName, FNativeClassHierarchyPluginModuleInfo> GetPluginModules();
 
 	/** Root level nodes corresponding to the root folders used by the Content Browser, eg) Classes_Engine, Classes_Game, etc */
 	TMap<FName, TSharedPtr<FNativeClassHierarchyNode>> RootNodes;
diff --git a/Engine/Source/Editor/ContentBrowser/Private/PathContextMenu.cpp b/Engine/Source/Editor/ContentBrowser/Private/PathContextMenu.cpp
--- a/Engine/Source/Editor/ContentBrowser/Private/PathContextMenu.cpp
+++ b/Engine/Source/Editor/ContentBrowser/Private/PathContextMenu.cpp
@@ -912,36 +912,7 @@ void FPathContextMenu::ExecuteSCCCheckIn()
 
 void FPathContextMenu::ExecuteSCCSync() const
 {
-	ISourceControlProvider& SourceControlProvider = ISourceControlModule::Get().GetProvider();
-
-	TArray<FString> RootPaths;
-	FPackageName::QueryRootContentPaths( RootPaths );
-
-	// find valid paths on disk
-	TArray<FString> PathsOnDisk;
-	for(const auto& SelectedPath : SelectedPaths)
-	{
-		// note: we make sure to terminate our path here so root directories map correctly.
-		const FString CompletePath = SelectedPath / "";
-
-		const bool bMatchesRoot = RootPaths.ContainsByPredicate([&](const FString& RootPath){ return CompletePath.StartsWith(RootPath); });
-		if(bMatchesRoot)
-		{
-			FString PathOnDisk = FPackageName::LongPackageNameToFilename(CompletePath);
-			PathsOnDisk.Add(PathOnDisk);
-		}
-	}
-
-	if ( PathsOnDisk.Num() > 0 )
-	{
-		TArray<FString> PackageNames;
-		GetPackageNamesInSelectedPaths(PackageNames);
-		ContentBrowserUtils::SyncPackagesFromSourceControl(PackageNames);
-	}
-	else
-	{
-		UE_LOG(LogContentBrowser, Warning, TEXT("Couldn't find any valid paths to sync."))
-	}
+	ContentBrowserUtils::SyncPathsFromSourceControl(SelectedPaths);
 }
 
 void FPathContextMenu::ExecuteSCCConnect() const
diff --git a/Engine/Source/Editor/ContentBrowser/Private/SAssetView.cpp b/Engine/Source/Editor/ContentBrowser/Private/SAssetView.cpp
--- a/Engine/Source/Editor/ContentBrowser/Private/SAssetView.cpp
+++ b/Engine/Source/Editor/ContentBrowser/Private/SAssetView.cpp
@@ -1783,11 +1783,13 @@ void SAssetView::RefreshSourceItems()
 		// If this is an engine folder, and we don't want to show them, remove
 		const bool IsHiddenEngineFolder = !bDisplayEngine && ContentBrowserUtils::IsEngineFolder(Item.PackagePath.ToString());
 		// If this is a plugin folder, and we don't want to show them, remove
-		const bool IsAHiddenPluginFolder = !bDisplayPlugins && ContentBrowserUtils::IsPluginFolder(Item.PackagePath.ToString());
+		const bool IsAHiddenGameProjectPluginFolder = !bDisplayPlugins && ContentBrowserUtils::IsPluginFolder(Item.PackagePath.ToString(), EPluginLoadedFrom::GameProject);
+		// If this is an engine plugin folder, and we don't want to show them, remove
+		const bool IsAHiddenEnginePluginFolder = (!bDisplayEngine || !bDisplayPlugins) && ContentBrowserUtils::IsPluginFolder(Item.PackagePath.ToString(), EPluginLoadedFrom::Engine);
 		// Do not show localized content folders.
 		const bool IsTheHiddenLocalizedContentFolder = !bDisplayL10N && ContentBrowserUtils::IsLocalizationFolder(Item.PackagePath.ToString());
 
-		const bool ShouldFilterOut = IsMainlyARedirector || IsHiddenEngineFolder || IsAHiddenPluginFolder || IsTheHiddenLocalizedContentFolder;
+		const bool ShouldFilterOut = IsMainlyARedirector || IsHiddenEngineFolder || IsAHiddenGameProjectPluginFolder || IsAHiddenEnginePluginFolder || IsTheHiddenLocalizedContentFolder;
 		if (ShouldFilterOut)
 		{
 			Items.RemoveAtSwap(AssetIdx);
diff --git a/Engine/Source/Editor/ContentBrowser/Private/SContentBrowser.cpp b/Engine/Source/Editor/ContentBrowser/Private/SContentBrowser.cpp
--- a/Engine/Source/Editor/ContentBrowser/Private/SContentBrowser.cpp
+++ b/Engine/Source/Editor/ContentBrowser/Private/SContentBrowser.cpp
@@ -27,6 +27,7 @@
 #include "EditorStyleSet.h"
 #include "EditorFontGlyphs.h"
 #include "Settings/ContentBrowserSettings.h"
+#include "Settings/EditorSettings.h"
 #include "Editor.h"
 #include "FileHelpers.h"
 #include "AssetRegistryModule.h"
@@ -710,6 +711,14 @@ void SContentBrowser::Construct( const FArguments& InArgs, const FName& InInstan
 			SourcesViewExpandClicked();
 		}
 	}
+	else
+	{
+		// in case we do not have a config, see what the global default settings are for the Sources Panel
+		if (!bSourcesViewExpanded && GetDefault<UEditorSettings>()->bOpenSourcesPanelByDefault)
+		{
+			SourcesViewExpandClicked();
+		}
+	}
 
 	// Bindings to manage history when items are deleted
 	FCollectionManagerModule& CollectionManagerModule = FCollectionManagerModule::GetModule();
diff --git a/Engine/Source/Editor/ContentBrowser/Private/SPathView.cpp b/Engine/Source/Editor/ContentBrowser/Private/SPathView.cpp
--- a/Engine/Source/Editor/ContentBrowser/Private/SPathView.cpp
+++ b/Engine/Source/Editor/ContentBrowser/Private/SPathView.cpp
@@ -780,7 +780,8 @@ TSharedPtr<struct FTreeItem> SPathView::AddRootItem( const FString& InFolderName
 	}
 
 	if ( (bDisplayEngine || !ContentBrowserUtils::IsEngineFolder(InFolderName)) && 
-		 (bDisplayPlugins || !ContentBrowserUtils::IsPluginFolder(InFolderName)) )
+ 		 ((bDisplayEngine && bDisplayPlugins) || !ContentBrowserUtils::IsPluginFolder(InFolderName, EPluginLoadedFrom::Engine)) &&
+		 (bDisplayPlugins || !ContentBrowserUtils::IsPluginFolder(InFolderName, EPluginLoadedFrom::GameProject)) )
 		{
 			const FText DisplayName = ContentBrowserUtils::GetRootDirDisplayName(InFolderName);
 			NewItem = MakeShareable( new FTreeItem(DisplayName, InFolderName, FString(TEXT("/")) + InFolderName, TSharedPtr<FTreeItem>()));
diff --git a/Engine/Source/Editor/CurveTableEditor/Private/CurveTableEditorHandle.h b/Engine/Source/Editor/CurveTableEditor/Private/CurveTableEditorHandle.h
--- a/Engine/Source/Editor/CurveTableEditor/Private/CurveTableEditorHandle.h
+++ b/Engine/Source/Editor/CurveTableEditor/Private/CurveTableEditorHandle.h
@@ -35,6 +35,12 @@ struct FCurveTableEditorHandle : public FCurveOwnerInterface
 	virtual void MakeTransactional() override;
 	virtual void OnCurveChanged(const TArray<FRichCurveEditInfo>& ChangedCurveEditInfos) override;
 	virtual bool IsValidCurve(FRichCurveEditInfo CurveInfo) override;
+	virtual TArray<const UObject*> GetOwners() const override
+	{ 
+		//Note this is read only so we return nothing
+		return TArray<const UObject*>(); 
+	}
+
 	//~ End FCurveOwnerInterface Interface.
 
 	/** Returns true if the curve is valid */
diff --git a/Engine/Source/Editor/DataTableEditor/Private/DataTableEditor.cpp b/Engine/Source/Editor/DataTableEditor/Private/DataTableEditor.cpp
--- a/Engine/Source/Editor/DataTableEditor/Private/DataTableEditor.cpp
+++ b/Engine/Source/Editor/DataTableEditor/Private/DataTableEditor.cpp
@@ -226,6 +226,33 @@ FSlateColor FDataTableEditor::GetRowTextColor(FName RowName) const
 	return FSlateColor::UseForeground();
 }
 
+FText FDataTableEditor::GetCellText(FDataTableEditorRowListViewDataPtr InRowDataPointer, int32 ColumnIndex) const
+{
+	if (InRowDataPointer.IsValid() && ColumnIndex < InRowDataPointer->CellData.Num())
+	{
+		return InRowDataPointer->CellData[ColumnIndex];
+	}
+
+	return FText();
+}
+
+FText FDataTableEditor::GetCellToolTipText(FDataTableEditorRowListViewDataPtr InRowDataPointer, int32 ColumnIndex) const
+{
+	FText TooltipText;
+
+	if (ColumnIndex < AvailableColumns.Num())
+	{
+		TooltipText = AvailableColumns[ColumnIndex]->DisplayName;
+	}
+
+	if (InRowDataPointer.IsValid() && ColumnIndex < InRowDataPointer->CellData.Num())
+	{
+		TooltipText = FText::Format(LOCTEXT("ColumnRowNameFmt", "{0}: {1}"), TooltipText, InRowDataPointer->CellData[ColumnIndex]);
+	}
+
+	return TooltipText;
+}
+
 FOptionalSize FDataTableEditor::GetRowNameColumnWidth() const
 {
 	return FOptionalSize(RowNameColumnWidth);
@@ -360,9 +387,9 @@ TSharedRef<SWidget> FDataTableEditor::MakeCellWidget(FDataTableEditorRowListView
 				SNew(STextBlock)
 				.TextStyle(FEditorStyle::Get(), "DataTableEditor.CellText")
 				.ColorAndOpacity(this, &FDataTableEditor::GetRowTextColor, InRowDataPtr->RowId)
-				.Text(InRowDataPtr->CellData[ColumnIndex])
+				.Text(this, &FDataTableEditor::GetCellText, InRowDataPtr, ColumnIndex)
 				.HighlightText(this, &FDataTableEditor::GetFilterText)
-				.ToolTipText(FText::Format(LOCTEXT("ColumnRowNameFmt", "{0}: {1}"), AvailableColumns[ColumnIndex]->DisplayName, InRowDataPtr->CellData[ColumnIndex]))
+				.ToolTipText(this, &FDataTableEditor::GetCellToolTipText, InRowDataPtr, ColumnIndex)
 			];
 	}
 
@@ -408,6 +435,8 @@ void FDataTableEditor::OnFilterTextChanged(const FText& InFilterText)
 void FDataTableEditor::RefreshCachedDataTable(const FName InCachedSelection, const bool bUpdateEvenIfValid)
 {
 	const UDataTable* Table = GetDataTable();
+	TArray<FDataTableEditorColumnHeaderDataPtr> PreviousColumns = AvailableColumns;
+
 	FDataTableEditorUtils::CacheDataTableForEditing(Table, AvailableColumns, AvailableRows);
 
 	// Update the desired width of the row names column
@@ -454,17 +483,20 @@ void FDataTableEditor::RefreshCachedDataTable(const FName InCachedSelection, con
 		}
 	}
 
-	ColumnNamesHeaderRow->ClearColumns();
-	for (int32 ColumnIndex = 0; ColumnIndex < AvailableColumns.Num(); ++ColumnIndex)
+	if (PreviousColumns != AvailableColumns)
 	{
-		const FDataTableEditorColumnHeaderDataPtr& ColumnData = AvailableColumns[ColumnIndex];
+		ColumnNamesHeaderRow->ClearColumns();
+		for (int32 ColumnIndex = 0; ColumnIndex < AvailableColumns.Num(); ++ColumnIndex)
+		{
+			const FDataTableEditorColumnHeaderDataPtr& ColumnData = AvailableColumns[ColumnIndex];
 
-		ColumnNamesHeaderRow->AddColumn(
-			SHeaderRow::Column(ColumnData->ColumnId)
-			.DefaultLabel(ColumnData->DisplayName)
-			.ManualWidth(TAttribute<float>::Create(TAttribute<float>::FGetter::CreateSP(this, &FDataTableEditor::GetColumnWidth, ColumnIndex)))
-			.OnWidthChanged(this, &FDataTableEditor::OnColumnResized, ColumnIndex)
+			ColumnNamesHeaderRow->AddColumn(
+				SHeaderRow::Column(ColumnData->ColumnId)
+				.DefaultLabel(ColumnData->DisplayName)
+				.ManualWidth(TAttribute<float>::Create(TAttribute<float>::FGetter::CreateSP(this, &FDataTableEditor::GetColumnWidth, ColumnIndex)))
+				.OnWidthChanged(this, &FDataTableEditor::OnColumnResized, ColumnIndex)
 			);
+		}
 	}
 
 	UpdateVisibleRows(InCachedSelection, bUpdateEvenIfValid);
@@ -507,11 +539,15 @@ void FDataTableEditor::UpdateVisibleRows(const FName InCachedSelection, const bo
 			}
 		}
 	}
+	
+	// Abort restoring the cached selection if data was changed while the user is selecting a different row
+	if (RowNamesListView->GetSelectedItems() == CellsListView->GetSelectedItems())
+	{
+		RowNamesListView->RequestListRefresh();
+		CellsListView->RequestListRefresh();
 
-	RowNamesListView->RequestListRefresh();
-	CellsListView->RequestListRefresh();
-
-	RestoreCachedSelection(InCachedSelection, bUpdateEvenIfValid);
+		RestoreCachedSelection(InCachedSelection, bUpdateEvenIfValid);
+	}
 }
 
 void FDataTableEditor::RestoreCachedSelection(const FName InCachedSelection, const bool bUpdateEvenIfValid)
diff --git a/Engine/Source/Editor/DataTableEditor/Private/DataTableEditor.h b/Engine/Source/Editor/DataTableEditor/Private/DataTableEditor.h
--- a/Engine/Source/Editor/DataTableEditor/Private/DataTableEditor.h
+++ b/Engine/Source/Editor/DataTableEditor/Private/DataTableEditor.h
@@ -86,6 +86,8 @@ private:
 	void OnFilterTextChanged(const FText& InFilterText);
 
 	FSlateColor GetRowTextColor(FName RowName) const;
+	FText GetCellText(FDataTableEditorRowListViewDataPtr InRowDataPointer, int32 ColumnIndex) const;
+	FText GetCellToolTipText(FDataTableEditorRowListViewDataPtr InRowDataPointer, int32 ColumnIndex) const;
 
 	TSharedRef<SVerticalBox> CreateContentBox();
 
diff --git a/Engine/Source/Editor/DetailCustomizations/Private/AnimTrailNodeDetails.cpp b/Engine/Source/Editor/DetailCustomizations/Private/AnimTrailNodeDetails.cpp
--- a/Engine/Source/Editor/DetailCustomizations/Private/AnimTrailNodeDetails.cpp
+++ b/Engine/Source/Editor/DetailCustomizations/Private/AnimTrailNodeDetails.cpp
@@ -98,6 +98,17 @@ void FAnimTrailNodeDetails::FTrailRelaxCurveEditor::ModifyOwner()
 	}
 }
 
+TArray<const UObject*> FAnimTrailNodeDetails::FTrailRelaxCurveEditor::GetOwners() const
+{
+	TArray<const UObject*> Owners;
+	if (GraphNodeOwner)
+	{
+		Owners.Add(GraphNodeOwner);
+	}
+
+	return Owners;
+}
+
 void FAnimTrailNodeDetails::FTrailRelaxCurveEditor::MakeTransactional()
 {
 	if(GraphNodeOwner)
diff --git a/Engine/Source/Editor/DetailCustomizations/Private/AnimTrailNodeDetails.h b/Engine/Source/Editor/DetailCustomizations/Private/AnimTrailNodeDetails.h
--- a/Engine/Source/Editor/DetailCustomizations/Private/AnimTrailNodeDetails.h
+++ b/Engine/Source/Editor/DetailCustomizations/Private/AnimTrailNodeDetails.h
@@ -26,6 +26,7 @@ public:
 		virtual TArray<FRichCurveEditInfoConst> GetCurves() const override;
 		virtual TArray<FRichCurveEditInfo> GetCurves() override;
 		virtual void ModifyOwner() override;
+		virtual TArray<const UObject*> GetOwners() const override;
 		virtual void MakeTransactional() override;
 		virtual void OnCurveChanged(const TArray<FRichCurveEditInfo>& ChangedCurveEditInfos) override { }
 		virtual bool IsValidCurve(FRichCurveEditInfo CurveInfo) override;
diff --git a/Engine/Source/Editor/DetailCustomizations/Private/BrushDetails.cpp b/Engine/Source/Editor/DetailCustomizations/Private/BrushDetails.cpp
--- a/Engine/Source/Editor/DetailCustomizations/Private/BrushDetails.cpp
+++ b/Engine/Source/Editor/DetailCustomizations/Private/BrushDetails.cpp
@@ -207,7 +207,7 @@ void FBrushDetails::CustomizeDetails( IDetailLayoutBuilder& DetailLayout )
 		BrushBuilderObject->SetFlags( RF_Transactional );
 	}
 
-	IDetailCategoryBuilder& BrushBuilderCategory = DetailLayout.EditCategory( "BrushSettings", FText::GetEmpty(), ECategoryPriority::Important );
+	IDetailCategoryBuilder& BrushBuilderCategory = DetailLayout.EditCategory( "BrushSettings", FText::GetEmpty() );
 
 	BrushBuilderCategory.AddProperty( GET_MEMBER_NAME_CHECKED(ABrush, BrushType) );
 	BrushBuilderCategory.AddCustomRow( LOCTEXT("BrushShape", "Brush Shape") )
diff --git a/Engine/Source/Editor/DetailCustomizations/Private/CameraDetails.cpp b/Engine/Source/Editor/DetailCustomizations/Private/CameraDetails.cpp
--- a/Engine/Source/Editor/DetailCustomizations/Private/CameraDetails.cpp
+++ b/Engine/Source/Editor/DetailCustomizations/Private/CameraDetails.cpp
@@ -72,10 +72,12 @@ void FCameraDetails::CustomizeDetails( IDetailLayoutBuilder& DetailLayout )
 		}
 	}
 
+	DetailLayout.EditCategory("Current Camera Settings", FText::GetEmpty(), ECategoryPriority::Important);
+
 	if (bCameraSettingsHidden == false)
 	{
 		IDetailCategoryBuilder& CameraCategory = DetailLayout.EditCategory( "CameraSettings", FText::GetEmpty(), ECategoryPriority::Important );
-	
+
 		// Organize the properties
 		CameraCategory.AddProperty(ProjectionModeProperty);
 
