Fix infinite loopin FUnloadedBlueprintData::GetClassWithin (introduced in 3357494)

diff --git a/Engine/Source/Editor/ClassViewer/Private/UnloadedBlueprintData.cpp b/Engine/Source/Editor/ClassViewer/Private/UnloadedBlueprintData.cpp
--- a/Engine/Source/Editor/ClassViewer/Private/UnloadedBlueprintData.cpp
+++ b/Engine/Source/Editor/ClassViewer/Private/UnloadedBlueprintData.cpp
@@ -75,6 +75,8 @@ const UClass* FUnloadedBlueprintData::GetClassWithin() const
 		{
 			return CurrentNode->Class->ClassWithin;
 		}
+
+		CurrentNode = CurrentNode->ParentNode.Pin();
 	}
 
 	return nullptr;
@@ -88,4 +90,4 @@ const TWeakPtr< class FClassViewerNode > FUnloadedBlueprintData::GetClassViewerN
 void FUnloadedBlueprintData::AddImplementedInterfaces(FString& InterfaceName)
 {
 	ImplementedInterfaces.Add(InterfaceName);
-}
\ No newline at end of file
+}
