Copying //UE4/Dev-Framework to //UE4/Dev-Main (Source: //UE4/Dev-Framework @ 4006870)
#rnx

============================
  MAJOR FEATURES & CHANGES
============================

Change 3967603 by [REDACTED_PERSON]

	Fix force feedback events being played on the correct controller if the OffsetPlayerGamepadIds feature is in use

Change 3969399 by [REDACTED_PERSON]

	Partially fix runtime failures for converted circular dependency test actors in a nativized QAGame.

Change 3970181 by [REDACTED_PERSON]

	Improved handling of circular dependencies during Blueprint nativization between converted and unconverted assets.

	Change summary:
	- Removed FGatherConvertedClassDependencies::MarkUnconvertedClassAsNecessary(); no longer in use.
	- Modified FFindAssetsToInclude::MaybeIncludeObjectAsDependency() to remove the case that was marking assets for stub class generation; no longer needed.
	- Modified FBlueprintNativeCodeGenModule::IsTargetedForReplacement() to detect, report on and handle circular references between converted and unconverted assets (ensuring that both sides are converted).

Change 3970185 by [REDACTED_PERSON]

	Fix initialization of special UMG struct types in nativized Blueprint ctor code.

	Change summary:
	- Modified FBackendHelperUMG::SpecialStructureConstructorUMG() to emit the proper constructor call based on the value of the 'ForcedTime' field. Also modified to emit updated constructor code for the FMovieSceneSegment struct's 'Range' member, which is now a TRange<> type.
	- Fixed a typo in the TBaseStructure<FInt32Range> case in FEmitDefaultValueHelper::SpecialStructureConstructor().

Change 3971520 by mason.seay

	Adjusted widget anchor to fix screen positioning

Change 3972611 by [REDACTED_PERSON]

	Fixed some of the vislogger macros using old 'Actor' rather than new 'LogOwner' parameter #UE4

Change 3973564 by [REDACTED_PERSON]

	Fix missing dependency on linker response file. Prevents target being relinked when build environment changes.

	- Mirrored from //UE4/Dev-Core (3973331).

Change 3974308 by [REDACTED_PERSON]

	Fix crash when choosing to delete multiple bookmarks selected in the Bookmarks tree view.

	Change summary:
	- Modified SBlueprintBookmarks::OnDeleteSelectedTreeViewItems() to make a local copy of the reference to the entry in the source array.
	- Modified FBlueprintEditor::RemoveBookmark() to save local config settings *after* modifying the persistent Bookmarks array (it was previously happening before).

Change 3975506 by [REDACTED_PERSON]

	Code and blueprints for subobject testing

Change 3977558 by [REDACTED_PERSON]

	Checking in missed header file

Change 3978137 by [REDACTED_PERSON]

Change 3978326 by [REDACTED_PERSON]

Change 3978393 by [REDACTED_PERSON]

	Fix FAssetData::IsUAsset comment to match what it actually does, old comment implied it was false for maps which is wrong

Change 3978837 by [REDACTED_PERSON]

	Make GetAutoEmitTermParameters and IsAutoCreateRefTerm static and optimize both

Change 3978838 by [REDACTED_PERSON]

	(4.19.1) PR #4480: Crash fix trying to open the console command on Windows (Contributed by [REDACTED_PERSON])

Change 3979156 by [REDACTED_PERSON]

	PR #4615: Update WindDirectionalSourceComponent.h (Contributed by [REDACTED_PERSON])

Change 3982289 by [REDACTED_PERSON]

	Fix several issues with unloaded blueprints and the class picker/viewer
	Added FName defines to FBlueprintTags for the asset registry tags used for unloaded blueprints, and removed the ParentClassPackage tag as it was redundant with ParentClass
	Fixed it so Soft Class Reference properties can be set back to none with the picker
	Removed FUnloadedClassDragDropOp as it did not work correctly and replaced uses with FAssetDragDrop, which was already used for content browser drags
	Changed internals of class viewer to properly store and use the paths to the actual generated class

Change 3982684 by [REDACTED_PERSON]

	Add DevelopmentAlwaysCook setting to the asset manager cook rule, this works like AlwaysCook in development and like NeverCook in production. DevelopmentCook works like Unknown in development

Change 3982762 by [REDACTED_PERSON]

	PR #4614: Fix UBTTask_RotateToFaceBBEntry when trying to face a position where i. (Contributed by [REDACTED_PERSON])

	I've improved the original PR by relaxing condition constraints, and removing some unnecessary vector normalization (since FVecotr::CosineAngle2D normalizes the input parameters anyway).

Change 3982866 by [REDACTED_PERSON]

	Defer Destroying an Actor if it is in the process of beginning play

Change 3982965 by [REDACTED_PERSON]

	PR #3567: add API declarations to BlueprintNodeHelpers (Contributed by [REDACTED_PERSON])

Change 3987399 by [REDACTED_PERSON]

	Remove unused ReloadObjectArc and unimplemented ArchiveReplaceArchetype

Change 3987956 by [REDACTED_PERSON]

	Fixed inconsistincies in BT category naming in BP #UE4

Change 3988308 by [REDACTED_PERSON]

	Fixed 'navmesh needs to be rebuild' displaying even after navigation building with navmesh auto building disabled #UE4

Change 3989309 by [REDACTED_PERSON]

	A support AI pluggin that links the aimodules #UE4

Change 3989311 by [REDACTED_PERSON]

	Moved the HTNPlanner plugin into the newly created Plugins/AI subfolder #UE4

Change 3989506 by mason.seay

	WIP Navigation Test Map

Change 3990826 by [REDACTED_PERSON]

	Fix existing placements of nativized Blueprint actors with a nonzero (legacy) root component scene transform at the CDO level.

	This corrects a long-standing issue that was introduced back in 2947488.

	Change summary:
	- Modified AActor::PostSpawnInitialize() to zero out RelativeLocation and RelativeRotation when spawning actors with a native scene root inherited from a converted Blueprint class.
	- Modified FNonativeComponentData to remove the HandledAsSpecialProperty() override for skipping the RelativeLocation and RelativeRotation properties on root components when emitting code to initialize the CDO for a converted Actor BP with a non-native root component.

Change 3993174 by mason.seay

	Cleared out unnecessary BP logic

Change 3994370 by [REDACTED_PERSON]

	Remove manual AI includes, the AI support plugin correctly pulls in those references now

Change 3995399 by [REDACTED_PERSON]

	When destroying the world clear out level collections

Change 3995731 by [REDACTED_PERSON]

	Fix UEngine::CommitMapChange issues with new streaming level state machine
	* Level needs to be added to array before determining current state
	* Should be adding, not setting streaming levels from fake world

Change 3996917 by [REDACTED_PERSON]

	ObjectWriter should not reset the soft object weak pointer, it should be a const operation

Change 3996921 by [REDACTED_PERSON]

Change 3997164 by [REDACTED_PERSON]

	Don't purge classes until they are recompiled and relink all child classes immediately after a blueprint generated class has changed. This allows iterative compilation logic to behave as it did before the compilation manager

Change 3997674 by paulo.souza

	Reverting camera setup changes from the renaming

Change 3997901 by mason.seay

	Updates to subobject test bp's

Change 3998028 by mason.seay

	Move logic to Event Graph

Change 3998051 by [REDACTED_PERSON]

	Adding QAGame ClassGroup to actively tested QAGame components for organization

Change 3999438 by [REDACTED_PERSON]

	Fix deterministic cooking issues by avoiding use of MakeUniqueObjectName when instantiating component templates

Change 4000265 by [REDACTED_PERSON]

	Checking in QAGame classes again to (hopefully) fix property issues

Change 4000339 by [REDACTED_PERSON]

	PR #4652: AttachComponent Message Fix (Contributed by [REDACTED_PERSON])

Change 4001338 by [REDACTED_PERSON]

	Fix reinstancing bug introduced in 3591939 which could cause some reinstanced objects to have references within them replaced. If the object in question was an actor it could take ownership of components owned by another actor and destroy them, resulting in a crash when the other actor tried to use its components

Change 4001400 by [REDACTED_PERSON]

	Per https://udn.unrealengine.com/questions/407745/childactorcomponent-and-modifications-to-component.html don't redirect or reset loaders when renaming components as part of the component instance data cache.

Change 4001530 by [REDACTED_PERSON]

	Use a clean default object to calculate the empty archive size.

Change 4002791 by [REDACTED_PERSON]

	Fix ensure when loading sublevel that has been redirected

Change 4002812 by mason.seay

	Added BP code to test Watching Sets and Maps

Change 4002947 by mason.seay

	More code for watch values testing

Change 4003059 by paulo.souza

	Fixes to volumetric lighting and sphere capture bakes

Change 4003509 by [REDACTED_PERSON]

	Reset bIsFakingTouched and bIsGameFakingTouch when the game viewport is unregistered. This prevents the PIE session from tainting editor session behavior when UInputSettings::bUseMouseForTouch is set

Change 4005596 by [REDACTED_PERSON]

	Fix random crashes when running multiple global blueprint search tabs in parallel just after editor launch.

	Change summary:
	- Refactored global FiB APIs to use thread-safe shared pointer types when referencing imaginary Blueprint data, which is a shared resource.
	- Modified FFindInBlueprintSearchManager::ContinueSearchQuery() to ensure that parsing imaginary Blueprint data only happens on a single thread.

Change 4005657 by [REDACTED_PERSON]

	Don't attach OldRoot component if it's not outered to the current actor

Change 4005709 by [REDACTED_PERSON]

	Ensure rather than assert on failed child promotion when destroying a non-root scene component within an Actor's scene component hierarchy.

Change 4005825 by paulo.souza

	New camera angle + touch rotation tests + "Action RPG" logos

Change 4006058 by [REDACTED_PERSON]

	Removing this as it is not needed.

Change 3980745 by [REDACTED_PERSON]

	Apply CL# 3977198 from Dev-Core
	Remove INI file override for editor stack size on Windows. This is rarely valid since editor targets share build products with other games by deafult. Fix to add linker response file as prerequisite exposed targets overriding this as a bug.

diff --git a/Engine/Plugins/AI/AISupport/Source/AISupportModule/Private/AISupportModule.cpp b/Engine/Plugins/AI/AISupport/Source/AISupportModule/Private/AISupportModule.cpp
new file mode 100644
--- /dev/null
+++ b/Engine/Plugins/AI/AISupport/Source/AISupportModule/Private/AISupportModule.cpp
@@ -0,0 +1,32 @@
+// Copyright 1998-2018 Epic Games, Inc. All Rights Reserved.
+
+#include "AISupportModule.h"
+#include "Modules/ModuleManager.h"
+
+#define LOCTEXT_NAMESPACE "AISupport"
+
+class FAISupportModule : public IAISupportModule
+{
+	// Begin IModuleInterface
+	virtual void StartupModule() override;
+	virtual void ShutdownModule() override;
+};
+
+IMPLEMENT_MODULE(FAISupportModule, AISupportModule)
+
+void FAISupportModule::StartupModule()
+{
+	// Called right after the module DLL has been loaded and the module object has been created
+	// Load dependent modules here, and they will be guaranteed to be available during ShutdownModule. ie:
+	// 		FModuleManager::Get().LoadModuleChecked(TEXT("HTTP"));
+}
+
+void FAISupportModule::ShutdownModule()
+{
+	// Called before the module is unloaded, right before the module object is destroyed.
+	// During normal shutdown, this is called in reverse order that modules finish StartupModule().
+	// This means that, as long as a module references dependent modules in it's StartupModule(), it
+	// can safely reference those dependencies in ShutdownModule() as well.
+}
+
+#undef LOCTEXT_NAMESPACE
diff --git a/Engine/Plugins/AI/HTNPlanner/Source/HTNPlanner/HTNPlanner.Build.cs b/Engine/Plugins/AI/HTNPlanner/Source/HTNPlanner/HTNPlanner.Build.cs
new file mode 100644
--- /dev/null
+++ b/Engine/Plugins/AI/HTNPlanner/Source/HTNPlanner/HTNPlanner.Build.cs
@@ -0,0 +1,56 @@
+// Copyright 1998-2018 Epic Games, Inc. All Rights Reserved.
+
+namespace UnrealBuildTool.Rules
+{
+    public class HTNPlanner : ModuleRules
+    {
+        public HTNPlanner(ReadOnlyTargetRules Target) : base(Target)
+        {
+            PublicIncludePaths.AddRange(
+                    new string[] {
+                        "Runtime/HTNPlanner/Public",
+                    }
+                    );
+
+            PrivateIncludePaths.AddRange(
+                new string[] {
+                        "Runtime/AIModule/Public",
+                        "Runtime/AIModule/Private",
+                        "Runtime/Engine/Private",
+                }
+                );
+
+            PublicDependencyModuleNames.AddRange(
+                new string[] {
+                        "Core",
+                        "CoreUObject",
+                        "Engine",
+                        "GameplayTags",
+                        "GameplayTasks",
+                        "AIModule"
+                }
+                );
+
+            DynamicallyLoadedModuleNames.AddRange(
+                new string[] {
+                    // ... add any modules that your module loads dynamically here ...
+                }
+                );
+
+            if (Target.bBuildEditor == true)
+            {
+                PrivateDependencyModuleNames.Add("UnrealEd");
+            }
+
+            if (Target.bBuildDeveloperTools || (Target.Configuration != UnrealTargetConfiguration.Shipping && Target.Configuration != UnrealTargetConfiguration.Test))
+            {
+                PrivateDependencyModuleNames.Add("GameplayDebugger");
+                PublicDefinitions.Add("WITH_GAMEPLAY_DEBUGGER=1");
+            }
+            else
+            {
+                PublicDefinitions.Add("WITH_GAMEPLAY_DEBUGGER=0");
+            }
+        }
+    }
+}
diff --git a/Engine/Plugins/AI/HTNPlanner/Source/HTNTestSuite/Private/HTNTest.cpp b/Engine/Plugins/AI/HTNPlanner/Source/HTNTestSuite/Private/HTNTest.cpp
new file mode 100644
--- /dev/null
+++ b/Engine/Plugins/AI/HTNPlanner/Source/HTNTestSuite/Private/HTNTest.cpp
@@ -0,0 +1,616 @@
+// Copyright 1998-2018 Epic Games, Inc. All Rights Reserved.
+
+#include "CoreMinimal.h"
+#include "HTNBuilder.h"
+#include "HTNPlanner.h"
+#include "AI/HTNBrainComponent.h"
+#include "MockHTN.h"
+#include "AITestsCommon.h"
+#include "Misc/App.h"
+#include "Debug/HTNDebug.h"
+
+#define LOCTEXT_NAMESPACE "AITestSuite_HTNTest"
+
+PRAGMA_DISABLE_OPTIMIZATION
+
+// missing tests:
+//  testing invalid conditions
+//	testing specific scenarios with known plans 
+//	test step by step generating the same plan as GeneratePlan
+//	using memory counters see if anything's leaking
+
+struct FHTNTestBase : public FAITestBase
+{
+	FHTNBuilder_Domain DomainBuilder;
+	FHTNWorldState WorldState;
+	FHTNPlanner Planner;
+
+	void PopulateWorldState()
+	{
+		for (int32 WSIndex = 0; WSIndex < int32(EMockHTNWorldState::MAX); ++WSIndex)
+		{
+			// setting every key to it's numerical index value
+			WorldState.SetValueUnsafe(WSIndex, WSIndex);
+		}
+	}
+
+	void PopulateDomain(const bool bCompile = true)
+	{
+		DomainBuilder.SetRootName(TEXT("Root"));
+		{
+			FHTNBuilder_CompositeTask& CompositeTaskBuilder = DomainBuilder.AddCompositeTask(TEXT("Root"));
+			{
+				FHTNBuilder_Method& MethodsBuilder = CompositeTaskBuilder.AddMethod(
+					TArray<FHTNCondition>({
+					FHTNCondition(EMockHTNWorldState::EnemyHealth, EHTNWorldStateCheck::Greater).SetRHSAsValue(0)
+					, FHTNCondition(EMockHTNWorldState::EnemyActor, EHTNWorldStateCheck::IsTrue)
+				}));
+				MethodsBuilder.AddTask(TEXT("AttackEnemy"));
+			}
+			{
+				FHTNBuilder_Method& MethodsBuilder = CompositeTaskBuilder.AddMethod();
+				MethodsBuilder.AddTask(TEXT("FindPatrolPoint"));
+				MethodsBuilder.AddTask(TEXT("NavigateToMoveDestination"));
+			}
+		}
+		{
+			FHTNBuilder_CompositeTask& CompositeTaskBuilder = DomainBuilder.AddCompositeTask(TEXT("AttackEnemy"));
+			{
+				FHTNBuilder_Method& MethodsBuilder = CompositeTaskBuilder.AddMethod(FHTNCondition(EMockHTNWorldState::HasWeapon, EHTNWorldStateCheck::IsTrue));
+				MethodsBuilder.AddTask(TEXT("NavigateToEnemy"));
+				MethodsBuilder.AddTask(TEXT("UseWeapon"));
+				MethodsBuilder.AddTask(TEXT("Root"));
+			}
+			{
+				FHTNBuilder_Method& MethodsBuilder = CompositeTaskBuilder.AddMethod();
+				MethodsBuilder.AddTask(TEXT("FindWeapon"));
+				MethodsBuilder.AddTask(TEXT("NavigateToWeapon"));
+				MethodsBuilder.AddTask(TEXT("PickUp"));
+				MethodsBuilder.AddTask(TEXT("AttackEnemy"));
+			}
+		}
+		{
+			FHTNBuilder_PrimitiveTask& PrimitiveTaskBuilder = DomainBuilder.AddPrimitiveTask(TEXT("FindPatrolPoint"));
+			PrimitiveTaskBuilder.SetOperator(EMockHTNTaskOperator::FindPatrolPoint, EMockHTNWorldState::MoveDestination);
+		}
+		{
+			FHTNBuilder_PrimitiveTask& PrimitiveTaskBuilder = DomainBuilder.AddPrimitiveTask(TEXT("FindWeapon"));
+			PrimitiveTaskBuilder.SetOperator(EMockHTNTaskOperator::FindWeapon, EMockHTNWorldState::PickupLocation);
+		}
+		{
+			FHTNBuilder_PrimitiveTask& PrimitiveTaskBuilder = DomainBuilder.AddPrimitiveTask(TEXT("NavigateToMoveDestination"));
+			PrimitiveTaskBuilder.SetOperator(EMockHTNTaskOperator::NavigateTo, EMockHTNWorldState::MoveDestination);	// Local Variables?
+			PrimitiveTaskBuilder.AddEffect(FHTNEffect(EMockHTNWorldState::CurrentLocation, EHTNWorldStateOperation::Set).SetRHSAsWSKey(EMockHTNWorldState::MoveDestination));
+		}
+		{
+			FHTNBuilder_PrimitiveTask& PrimitiveTaskBuilder = DomainBuilder.AddPrimitiveTask(TEXT("NavigateToEnemy"));
+			PrimitiveTaskBuilder.SetOperator(EMockHTNTaskOperator::NavigateTo, EMockHTNWorldState::EnemyActor);
+			PrimitiveTaskBuilder.AddEffect(FHTNEffect(EMockHTNWorldState::CurrentLocation, EHTNWorldStateOperation::Set).SetRHSAsWSKey(EMockHTNWorldState::EnemyActor));
+			PrimitiveTaskBuilder.AddEffect(FHTNEffect(EMockHTNWorldState::CanSeeEnemy, EHTNWorldStateOperation::Set).SetRHSAsValue(1));
+		}
+		{
+			FHTNBuilder_PrimitiveTask& PrimitiveTaskBuilder = DomainBuilder.AddPrimitiveTask(TEXT("NavigateToWeapon"));
+			PrimitiveTaskBuilder.SetOperator(EMockHTNTaskOperator::NavigateTo, EMockHTNWorldState::PickupLocation);
+			PrimitiveTaskBuilder.AddEffect(FHTNEffect(EMockHTNWorldState::CurrentLocation, EHTNWorldStateOperation::Set).SetRHSAsWSKey(EMockHTNWorldState::PickupLocation));
+		}
+		{
+			FHTNBuilder_PrimitiveTask& PrimitiveTaskBuilder = DomainBuilder.AddPrimitiveTask(TEXT("PickUp"));
+			PrimitiveTaskBuilder.SetOperator(EMockHTNTaskOperator::PickUp, EMockHTNWorldState::PickupLocation);
+			PrimitiveTaskBuilder.AddEffect(FHTNEffect(EMockHTNWorldState::HasWeapon, EHTNWorldStateOperation::Set).SetRHSAsValue(1));
+		}
+		{
+			FHTNBuilder_PrimitiveTask& PrimitiveTaskBuilder = DomainBuilder.AddPrimitiveTask(TEXT("UseWeapon"));
+			PrimitiveTaskBuilder.SetOperator(EMockHTNTaskOperator::UseWeapon, EMockHTNWorldState::EnemyActor);
+			PrimitiveTaskBuilder.AddEffect(FHTNEffect(EMockHTNWorldState::Ammo, EHTNWorldStateOperation::Decrease).SetRHSAsValue(1));
+			PrimitiveTaskBuilder.AddEffect(FHTNEffect(EMockHTNWorldState::EnemyHealth, EHTNWorldStateOperation::Decrease).SetRHSAsValue(1));
+		}
+
+		if (bCompile)
+		{
+			DomainBuilder.Compile();
+		}
+	}
+};
+
+//////////////////////////////////////////////////////////////////////////
+
+struct FAITest_HTNDomainBuilderBasics : public FAITestBase
+{
+	FAITest_HTNDomainBuilderBasics()
+	{
+		FHTNBuilder_Domain DomainBuilder;
+
+		Test(TEXT("Initially DomainBuilder instance should be empty"), DomainBuilder.CompositeTasks.Num() == 0 && DomainBuilder.PrimitiveTasks.Num() == 0);
+
+		const uint32 CompointTasksCount = 5;
+		for (int32 CompositeTaskIndex = 0; CompositeTaskIndex < CompointTasksCount; ++CompositeTaskIndex)
+		{
+			FHTNBuilder_CompositeTask& CompositeTaskBuilder = DomainBuilder.AddCompositeTask(*FString::Printf(TEXT("c_%d"), CompositeTaskIndex));
+			for (int32 MethodIndex = 0; MethodIndex < (CompositeTaskIndex % 3); ++MethodIndex)
+			{
+				FHTNBuilder_Method* MethodsBuilder = nullptr;
+
+				if ((CompositeTaskIndex + MethodIndex) != 0)
+				{
+					TArray<FHTNCondition> Conditions;
+					for (int32 ConditionIndex = 0; ConditionIndex < (CompositeTaskIndex % 3); ++ConditionIndex)
+					{
+						Conditions.Add(FHTNCondition(ConditionIndex, EHTNWorldStateCheck(ConditionIndex)));
+					}
+
+					MethodsBuilder = &CompositeTaskBuilder.AddMethod(Conditions);
+				}
+				else
+				{
+					// this is testing the non-array method
+					MethodsBuilder = &CompositeTaskBuilder.AddMethod(FHTNCondition(0, EHTNWorldStateCheck(0)));
+				}
+
+				for (int32 TaskIndex = 0; TaskIndex < (CompositeTaskIndex % 4); ++TaskIndex)
+				{
+					MethodsBuilder->AddTask(*FString::Printf(TEXT("t_%d"), CompositeTaskIndex));
+				}
+			}
+		}
+
+		for (int32 CompositeTaskIndex = 0; CompositeTaskIndex < DomainBuilder.CompositeTasks.Num(); ++CompositeTaskIndex)
+		{
+			FHTNBuilder_CompositeTask* CompositeTaskBuilder = DomainBuilder.CompositeTasks.Find(*FString::Printf(TEXT("c_%d"), CompositeTaskIndex));
+			Test(FString::Printf(TEXT("Failed to find Composite task c_%d that has just been added"), CompositeTaskIndex), CompositeTaskBuilder != nullptr);
+
+			if (CompositeTaskBuilder != nullptr)
+			{
+				Test(FString::Printf(TEXT("Method count mismatch for c_%d"), CompositeTaskIndex), CompositeTaskBuilder->Methods.Num() == (CompositeTaskIndex % 3));
+				for (int32 MethodIndex = 0; MethodIndex < CompositeTaskBuilder->Methods.Num(); ++MethodIndex)
+				{
+					if ((CompositeTaskIndex + MethodIndex) != 0)
+					{
+						Test(FString::Printf(TEXT("Condition count mismatch for c_%d[%d] method (array path)"), CompositeTaskIndex, MethodIndex)
+							, CompositeTaskBuilder->Methods[MethodIndex].Conditions.Num() == (CompositeTaskIndex % 3));
+					}
+					else
+					{
+						Test(FString::Printf(TEXT("Condition count mismatch for c_%d[%d] method (single instance)"), CompositeTaskIndex, MethodIndex)
+							, CompositeTaskBuilder->Methods[MethodIndex].Conditions.Num() == 1);
+					}
+
+					Test(FString::Printf(TEXT("Task count mismatch for c_%d[%d] method"), CompositeTaskIndex, MethodIndex)
+						, CompositeTaskBuilder->Methods[MethodIndex].Tasks.Num() == (CompositeTaskIndex % 4));
+				}
+			}			
+		}
+
+		const uint32 PrimitiveTasksCount = 4;
+		for (int32 PrimitiveTaskIndex = 0; PrimitiveTaskIndex < PrimitiveTasksCount; ++PrimitiveTaskIndex)
+		{
+			FHTNBuilder_PrimitiveTask& PrimitiveTaskBuilder = DomainBuilder.AddPrimitiveTask(*FString::Printf(TEXT("p_%d"), PrimitiveTaskIndex));
+			PrimitiveTaskBuilder.SetOperator(PrimitiveTaskIndex, PrimitiveTaskIndex * 2);
+			for (int32 EffectIndex = 0; EffectIndex < (PrimitiveTaskIndex % 3); ++EffectIndex)
+			{
+				PrimitiveTaskBuilder.AddEffect(FHTNEffect(EffectIndex, EHTNWorldStateOperation::Set));
+			}
+		}
+
+		Test(TEXT("Wrong number of primitive tasks added"), DomainBuilder.PrimitiveTasks.Num() == PrimitiveTasksCount);
+		for (int32 PrimitiveTaskIndex = 0; PrimitiveTaskIndex < PrimitiveTasksCount; ++PrimitiveTaskIndex)
+		{
+			FHTNBuilder_PrimitiveTask* PrimitiveTaskBuilder = DomainBuilder.PrimitiveTasks.Find(*FString::Printf(TEXT("p_%d"), PrimitiveTaskIndex));
+			Test(FString::Printf(TEXT("Failed to find primitive task p_%d that has just been added"), PrimitiveTaskIndex), PrimitiveTaskBuilder != nullptr);
+
+			if (PrimitiveTaskBuilder)
+			{
+				Test(FString::Printf(TEXT("Primitive task p_%d operator is wrong"), PrimitiveTasksCount), PrimitiveTaskBuilder->ActionID == PrimitiveTaskIndex);
+				Test(FString::Printf(TEXT("Primitive task p_%d effects count is wrong"), PrimitiveTasksCount), PrimitiveTaskBuilder->Effects.Num() == (PrimitiveTaskIndex % 3));
+			}
+		}
+	}
+};
+IMPLEMENT_AI_LATENT_TEST(FAITest_HTNDomainBuilderBasics, "System.Engine.AI.HTN.DomainBuilderBasics")
+
+struct FAITest_HTNBuildDomain : public FHTNTestBase
+{
+	FAITest_HTNBuildDomain()
+	{
+		PopulateDomain();
+
+		Test(TEXT("DomainBuilder stores wrong number of primitive tasks"), DomainBuilder.PrimitiveTasks.Num() == 7);
+		Test(TEXT("DomainBuilder stores wrong number of Composite tasks"), DomainBuilder.CompositeTasks.Num() == 2);
+		
+		FHTNBuilder_CompositeTask* RootBuilder = DomainBuilder.GetRootAsCompositeTask();
+		Test(TEXT("Root task should be set"), RootBuilder != nullptr);
+		if (RootBuilder && RootBuilder->Methods.Num() > 0)
+		{
+			Test(TEXT("Root task\'s method [0] should be configured to have two conditions"), RootBuilder->Methods[0].Conditions.Num() == 2);
+		}
+		
+		TArray<FHTNBuilder_CompositeTask> CompositeTasks;
+		DomainBuilder.CompositeTasks.GenerateValueArray(CompositeTasks);
+		Test(TEXT("DomainBuilder stores wrong number of methods for the first Composite task"), CompositeTasks[0].Methods.Num() == 2);
+		Test(TEXT("DomainBuilder stores wrong number of methods for the second Composite task"), CompositeTasks[1].Methods.Num() == 2);
+	}
+};
+IMPLEMENT_AI_LATENT_TEST(FAITest_HTNBuildDomain, "System.Engine.AI.HTN.BuildDomain")
+
+struct FAITest_HTNPlanning : public FHTNTestBase
+{
+	FAITest_HTNPlanning()
+	{
+		FHTNResult Result;
+		
+		Planner.GeneratePlan(*(DomainBuilder.DomainInstance), WorldState, Result);
+		Test(TEXT("Planning with an empty domain should result in an empty plan"), Result.TaskIDs.Num() == 0);
+
+		PopulateDomain();
+
+		Planner.GeneratePlan(*(DomainBuilder.DomainInstance), WorldState, Result);
+
+		TArray<FHTNBuilder_CompositeTask> CompositeTasks;
+		DomainBuilder.CompositeTasks.GenerateValueArray(CompositeTasks);
+		// Patrol plan
+		Test(TEXT("Patrol plan should be same length as the last methods of the root Composite task"), CompositeTasks[0].Methods.Last().Tasks.Num() == Result.TaskIDs.Num());
+		for (int32 TaskIndex = 0; TaskIndex < Result.TaskIDs.Num(); ++TaskIndex)
+		{
+			const FHTNPolicy::FTaskID TaskID = Result.TaskIDs[TaskIndex];
+			Test(TEXT("Patrol plan element mismatch"), TaskID == DomainBuilder.DomainInstance->FindTaskID(CompositeTasks[0].Methods.Last().Tasks[TaskIndex]));
+		}
+
+		const FHTNDomain EmptyDomain;
+		Planner.GeneratePlan(EmptyDomain, WorldState, Result);
+		Test(TEXT("Reusing previous planning result with an empty domain should result in an empty plan"), Result.TaskIDs.Num() == 0);
+
+		WorldState.ApplyEffect(FHTNEffect(EMockHTNWorldState::EnemyHealth, EHTNWorldStateOperation::Set).SetRHSAsValue(1));
+		WorldState.ApplyEffect(FHTNEffect(EMockHTNWorldState::EnemyActor, EHTNWorldStateOperation::Set).SetRHSAsValue(1));
+		Planner.GeneratePlan(*(DomainBuilder.DomainInstance), WorldState, Result);
+		Test(TEXT(""), Result.TaskIDs.Num() > 0);
+	}
+};
+IMPLEMENT_AI_LATENT_TEST(FAITest_HTNPlanning, "System.Engine.AI.HTN.Planning")
+
+struct FAITest_HTNPlanningRollback : public FHTNTestBase
+{
+	FAITest_HTNPlanningRollback()
+	{		
+		FHTNResult Result;
+
+		// build a domain that will force rolling back
+		// first method should get accepted in the first planner step
+		//		then one of the tasks it consists of should be a Composite task, that fails its condition
+		// Note: WorldState is populated with 0 values and has 128 keys (by default)
+		{
+			FHTNBuilder_CompositeTask& CompositeTaskBuilder = DomainBuilder.AddCompositeTask(NAME_None);	// root
+			{
+				FHTNBuilder_Method& MethodsBuilder = CompositeTaskBuilder.AddMethod();
+				MethodsBuilder.AddTask(TEXT("FailedComposite"));
+			}
+			{
+				FHTNBuilder_Method& MethodsBuilder = CompositeTaskBuilder.AddMethod();
+				MethodsBuilder.AddTask(TEXT("SuccessfulComposite"));
+			}
+		}
+		{
+			FHTNBuilder_CompositeTask& CompositeTaskBuilder = DomainBuilder.AddCompositeTask(TEXT("FailedComposite"));
+			{
+				FHTNBuilder_Method& MethodsBuilder = CompositeTaskBuilder.AddMethod(FHTNCondition(0, EHTNWorldStateCheck::Greater).SetRHSAsValue(0));
+				MethodsBuilder.AddTask(TEXT("DummyPrimitive1"));
+			}
+		}
+		{
+			FHTNBuilder_CompositeTask& CompositeTaskBuilder = DomainBuilder.AddCompositeTask(TEXT("SuccessfulComposite"));
+			{
+				FHTNBuilder_Method& MethodsBuilder = CompositeTaskBuilder.AddMethod(FHTNCondition(0, EHTNWorldStateCheck::Equal).SetRHSAsValue(0));
+				MethodsBuilder.AddTask(TEXT("DummyPrimitive2"));
+			}
+		}
+		DomainBuilder.AddPrimitiveTask(TEXT("DummyPrimitive1"));
+		const FName DummyPrimitiveName2 = TEXT("DummyPrimitive2");
+		DomainBuilder.AddPrimitiveTask(DummyPrimitiveName2);
+		
+		DomainBuilder.Compile();
+
+		Planner.GeneratePlan(*(DomainBuilder.DomainInstance), WorldState, Result);
+		Test(TEXT("First Rollback plan should consist of one task, DummyPrimitive2"), Result.TaskIDs.Num() == 1 
+			&& FHTNDebug::GetTaskName(DomainBuilder, Result.TaskIDs[0]) == DummyPrimitiveName2);
+	}
+};
+IMPLEMENT_AI_LATENT_TEST(FAITest_HTNPlanningRollback, "System.Engine.AI.HTN.PlanningRollback")
+
+struct FAITest_HTNDecompileDomain : public FHTNTestBase
+{
+	FAITest_HTNDecompileDomain()
+	{
+		Test(TEXT("Compiling an empty domain is allowed"), DomainBuilder.Compile());
+
+		PopulateDomain();
+
+		FHTNBuilder_Domain DomainBuilder2(DomainBuilder.DomainInstance);
+		DomainBuilder2.Decompile();
+		const FString OriginalDescription = DomainBuilder.GetDebugDescription();
+		const FString DecompiledDescription = DomainBuilder2.GetDebugDescription();
+
+		Test(TEXT("Decompilation should result in identical DomainBuilder"), OriginalDescription.Equals(DecompiledDescription));
+	}
+};
+IMPLEMENT_AI_LATENT_TEST(FAITest_HTNDecompileDomain, "System.Engine.AI.HTN.DomainDecompilation")
+
+struct FAITest_HTNDomainCompilationIssues : public FHTNTestBase
+{
+	FAITest_HTNDomainCompilationIssues()
+	{
+		const FName MissingTaskName = TEXT("MissingTask");
+
+		FHTNBuilder_CompositeTask& CompositeTaskBuilder = DomainBuilder.AddCompositeTask(TEXT("Root"));	// root
+		{
+			FHTNBuilder_Method& MethodsBuilder = CompositeTaskBuilder.AddMethod();
+			MethodsBuilder.AddTask(MissingTaskName);
+		}
+
+		Test(TEXT("Domain with missing tasks should not compile"), DomainBuilder.Compile() == false);
+		Test(TEXT("Domain should be empty after a failed compilation"), DomainBuilder.DomainInstance->IsEmpty());
+
+		DomainBuilder.AddPrimitiveTask(MissingTaskName);
+		Test(TEXT("After adding missing task domain should compile just fine"), DomainBuilder.Compile());
+		Test(TEXT("Domain should not be empty after a successful compilation"), DomainBuilder.DomainInstance->IsEmpty() == false);
+	}
+};
+IMPLEMENT_AI_LATENT_TEST(FAITest_HTNDomainCompilationIssues, "System.Engine.AI.HTN.DomainDecompilationIssues")
+
+struct FAITest_HTNWorldRepresentation : public FHTNTestBase
+{
+	FAITest_HTNWorldRepresentation()
+	{
+		for (int32 WSIndex = 0; WSIndex < int32(EMockHTNWorldState::MAX); ++WSIndex)
+		{
+			FHTNPolicy::FWSValue Value;
+			Test(TEXT("Retrieving known values from the WorldState instance"), WorldState.GetValue(WSIndex, Value) && Value == FHTNPolicy::DefaultValue);
+		}
+
+		PopulateWorldState();
+		
+		const uint32 ReferenceValue = 3;
+		for (int32 WSIndex = 0; WSIndex < int32(EMockHTNWorldState::MAX); ++WSIndex)
+		{
+			FHTNPolicy::FWSValue Value;
+			Test(TEXT("Retrieving known values from the WorldState instance"), WorldState.GetValue(WSIndex, Value) && Value == WSIndex);
+
+			for (int32 OpIndex = 0; OpIndex < int32(EHTNWorldStateCheck::MAX); ++OpIndex)
+			{
+				const EHTNWorldStateCheck OpCode = EHTNWorldStateCheck(OpIndex);
+				bool bExpectedResult = false;
+
+				switch (OpCode)
+				{
+				case EHTNWorldStateCheck::Less:
+					bExpectedResult = Value < ReferenceValue;
+					break;
+				case EHTNWorldStateCheck::LessOrEqual:
+					bExpectedResult = Value <= ReferenceValue;
+					break;
+				case EHTNWorldStateCheck::Equal:
+					bExpectedResult = Value == ReferenceValue;
+					break;
+				case EHTNWorldStateCheck::NotEqual:
+					bExpectedResult = Value != ReferenceValue;
+					break;
+				case EHTNWorldStateCheck::GreaterOrEqual:
+					bExpectedResult = Value >= ReferenceValue;
+					break;
+				case EHTNWorldStateCheck::Greater:
+					bExpectedResult = Value > ReferenceValue;
+					break;
+				case EHTNWorldStateCheck::IsTrue:
+					bExpectedResult = (Value != 0);
+					break;
+				default:
+					Test(TEXT("Unhanled operation ID!"), true);
+				}
+
+				FString Message = FString::Printf(TEXT("Testing %s on $d"), *FHTNDebug::HTNWorldStateCheckToString(OpCode), Value);
+				Test(Message, WorldState.CheckCondition(FHTNCondition(WSIndex, OpCode).SetRHSAsValue(ReferenceValue)) == bExpectedResult);
+			}
+		}
+	}
+};
+IMPLEMENT_AI_LATENT_TEST(FAITest_HTNWorldRepresentation, "System.Engine.AI.HTN.WorldRepresentation")
+
+struct FAITest_HTNCondition : public FHTNTestBase
+{
+	FAITest_HTNCondition()
+	{
+		PopulateWorldState();
+
+		const uint32 ReferenceValue = 3;
+
+		for (int32 WSIndex = 0; WSIndex < int32(EMockHTNWorldState::MAX); ++WSIndex)
+		{
+			for (int32 Value = 0; Value < int32(EMockHTNWorldState::MAX); ++Value)
+			{
+				const FHTNPolicy::FWSValue AsValue = FHTNPolicy::FWSValue(Value);
+				const FHTNPolicy::FWSKey AsKey = FHTNPolicy::FWSKey(Value);
+
+				Test(FString::Printf(TEXT("Condition WS[%d] < %d"), WSIndex, AsValue)
+					, WorldState.CheckCondition(FHTNCondition(WSIndex, EHTNWorldStateCheck::Less).SetRHSAsValue(AsValue)) == (WSIndex < AsValue));
+				Test(FString::Printf(TEXT("Condition WS[%d] <= %d"), WSIndex, AsValue)
+					, WorldState.CheckCondition(FHTNCondition(WSIndex, EHTNWorldStateCheck::LessOrEqual).SetRHSAsValue(AsValue)) == (WSIndex <= AsValue));
+				Test(FString::Printf(TEXT("Condition WS[%d] == %d"), WSIndex, AsValue)
+					, WorldState.CheckCondition(FHTNCondition(WSIndex, EHTNWorldStateCheck::Equal).SetRHSAsValue(AsValue)) == (WSIndex == AsValue));
+				Test(FString::Printf(TEXT("Condition WS[%d] != %d"), WSIndex, AsValue)
+					, WorldState.CheckCondition(FHTNCondition(WSIndex, EHTNWorldStateCheck::NotEqual).SetRHSAsValue(AsValue)) == (WSIndex != AsValue));
+				Test(FString::Printf(TEXT("Condition WS[%d] >= %d"), WSIndex, AsValue)
+					, WorldState.CheckCondition(FHTNCondition(WSIndex, EHTNWorldStateCheck::GreaterOrEqual).SetRHSAsValue(AsValue)) == (WSIndex >= AsValue));
+				Test(FString::Printf(TEXT("Condition WS[%d] > %d"), WSIndex, AsValue)
+					, WorldState.CheckCondition(FHTNCondition(WSIndex, EHTNWorldStateCheck::Greater).SetRHSAsValue(AsValue)) == (WSIndex > AsValue));
+
+				Test(FString::Printf(TEXT("Condition WS[%d] < WS[%d]"), WSIndex, AsKey)
+					, WorldState.CheckCondition(FHTNCondition(WSIndex, EHTNWorldStateCheck::Less).SetRHSAsWSKey(AsKey)) == (WSIndex < AsKey));
+				Test(FString::Printf(TEXT("Condition WS[%d] <= WS[%d]"), WSIndex, AsKey)
+					, WorldState.CheckCondition(FHTNCondition(WSIndex, EHTNWorldStateCheck::LessOrEqual).SetRHSAsWSKey(AsKey)) == (WSIndex <= AsKey));
+				Test(FString::Printf(TEXT("Condition WS[%d] == WS[%d]"), WSIndex, AsKey)
+					, WorldState.CheckCondition(FHTNCondition(WSIndex, EHTNWorldStateCheck::Equal).SetRHSAsWSKey(AsKey)) == (WSIndex == AsKey));
+				Test(FString::Printf(TEXT("Condition WS[%d] != WS[%d]"), WSIndex, AsKey)
+					, WorldState.CheckCondition(FHTNCondition(WSIndex, EHTNWorldStateCheck::NotEqual).SetRHSAsWSKey(AsKey)) == (WSIndex != AsKey));
+				Test(FString::Printf(TEXT("Condition WS[%d] >= WS[%d]"), WSIndex, AsKey)
+					, WorldState.CheckCondition(FHTNCondition(WSIndex, EHTNWorldStateCheck::GreaterOrEqual).SetRHSAsWSKey(AsKey)) == (WSIndex >= AsKey));
+				Test(FString::Printf(TEXT("Condition WS[%d] > WS[%d]"), WSIndex, AsKey)
+					, WorldState.CheckCondition(FHTNCondition(WSIndex, EHTNWorldStateCheck::Greater).SetRHSAsWSKey(AsKey)) == (WSIndex > AsKey));
+			}
+		}
+	}
+};
+IMPLEMENT_AI_LATENT_TEST(FAITest_HTNCondition, "System.Engine.AI.HTN.Condition")
+
+struct FAITest_HTNMethodSelection : public FAITestBase 
+{
+	FAITest_HTNMethodSelection()
+	{
+		FHTNDomain Domain;
+
+		//FHTNCompositeTask& RootCompositeTask = Domain.AddCompositeTask(TEXT("Root"));
+		// do planning and come up with empty plan
+
+		// INJECTION
+	}
+};
+IMPLEMENT_AI_LATENT_TEST(FAITest_HTNMethodSelection, "System.Engine.AI.HTN.MethodSelection")
+
+struct FAITest_HTNTrivialPlanning : public FHTNTestBase 
+{
+	FAITest_HTNTrivialPlanning()
+	{
+		//FHTNDomain Domain;
+
+		TArray<FHTNCondition> Conditions;
+		//FHTNMethodSelection MethodSelection(0, Conditions);
+
+
+		//FHTNCompositeTask& RootCompositeTask = Domain.AddCompositeTask(TEXT("Root"));
+		// do planning and come up with empty plan
+	}
+};
+IMPLEMENT_AI_LATENT_TEST(FAITest_HTNTrivialPlanning, "System.Engine.AI.HTN.TrivialPlanning")
+
+struct FAITest_HTNExtendingDomain : public FAITestBase
+{
+	FAITest_HTNExtendingDomain()
+	{
+		//FHTNDomain Domain;
+
+		//FHTNCompositeTask& RootCompositeTask = Domain.AddCompositeTask(TEXT("Root"));
+		// do planning and come up with empty plan
+
+		// INJECTION
+	}
+};
+IMPLEMENT_AI_LATENT_TEST(FAITest_HTNExtendingDomain, "System.Engine.AI.HTN.ExtendingDomain")
+
+struct FAITest_HTNCustomWSCheck : public FHTNTestBase
+{
+	static bool bCheckFunctionRun;
+	static bool CustomCheck(const FHTNPolicy::FWSValue* Values, const FHTNCondition& Condition)
+	{
+		bCheckFunctionRun = true;
+		return true;
+	}
+
+	FAITest_HTNCustomWSCheck()
+	{
+		const uint32 CustomCheckID = FHTNWorldStateOperations::RegisterCustomCheckType(&FAITest_HTNCustomWSCheck::CustomCheck, TEXT("CustomCheck")); 
+
+		FHTNBuilder_CompositeTask& CompositeTaskBuilder = DomainBuilder.AddCompositeTask(NAME_None);	// root
+		{
+			FHTNBuilder_Method& MethodsBuilder = CompositeTaskBuilder.AddMethod(FHTNCondition(0, CustomCheckID));
+			MethodsBuilder.AddTask(TEXT("Task1"));
+		}
+		FHTNBuilder_PrimitiveTask& PrimitiveTask = DomainBuilder.AddPrimitiveTask(TEXT("Task1"));
+		
+		DomainBuilder.Compile();
+		
+		FHTNResult Result;
+		Planner.GeneratePlan(*DomainBuilder.DomainInstance, WorldState, Result);
+
+		Test(TEXT("Custom check has been exeduted"), bCheckFunctionRun);
+		Test(TEXT("The custom check should allow for construction of the plan"), Result.TaskIDs.Num() == 1);
+	}
+};
+bool FAITest_HTNCustomWSCheck::bCheckFunctionRun = false;
+IMPLEMENT_AI_LATENT_TEST(FAITest_HTNCustomWSCheck, "System.Engine.AI.HTN.CustomWSCheck")
+
+struct FAITest_HTNCustomWSOperation : public FHTNTestBase
+{
+	static uint32 StaticValue;
+	static void CustomOperation(FHTNPolicy::FWSValue* Values, const FHTNEffect& Effect)
+	{
+		++StaticValue;
+		Values[Effect.KeyLeftHand] = StaticValue * 1024;
+	}
+
+	FAITest_HTNCustomWSOperation()
+	{
+		const uint32 CustomOperationID = FHTNWorldStateOperations::RegisterCustomOperationType(&FAITest_HTNCustomWSOperation::CustomOperation, TEXT("CustomOperation"));
+
+		FHTNBuilder_CompositeTask& CompositeTaskBuilder = DomainBuilder.AddCompositeTask(NAME_None);	// root
+		{
+			FHTNBuilder_Method& MethodsBuilder = CompositeTaskBuilder.AddMethod();
+			MethodsBuilder.AddTask(TEXT("Task1"));
+		}
+		FHTNBuilder_PrimitiveTask& PrimitiveTask = DomainBuilder.AddPrimitiveTask(TEXT("Task1"));
+		PrimitiveTask.AddEffect(FHTNEffect(0, CustomOperationID));
+		PrimitiveTask.AddEffect(FHTNEffect(2, CustomOperationID));
+
+		DomainBuilder.Compile();
+
+		FHTNResult Result;
+		Planner.GeneratePlan(*DomainBuilder.DomainInstance, WorldState, Result);
+
+		Test(TEXT("Custom Operation has been exeduted"), StaticValue == 2);
+		Test(TEXT("Checking custom effect on key 0"), Planner.GetWorldState().GetValueUnsafe(0) == 1024);
+		Test(TEXT("Checking custom effect on key 0"), Planner.GetWorldState().GetValueUnsafe(2) == 1024 * 2);
+	}
+};
+uint32 FAITest_HTNCustomWSOperation::StaticValue = 0;
+IMPLEMENT_AI_LATENT_TEST(FAITest_HTNCustomWSOperation, "System.Engine.AI.HTN.CustomWSOperation")
+
+//	compare contents of TaskIDs in HTNResult to the operators it contains
+struct FAITest_HTNOperatorsOfGeneratedPlan : public FHTNTestBase
+{
+	FAITest_HTNOperatorsOfGeneratedPlan()
+	{
+		FHTNBuilder_CompositeTask& CompositeTaskBuilder = DomainBuilder.AddCompositeTask(NAME_None);
+		{
+			FHTNBuilder_Method& MethodsBuilder = CompositeTaskBuilder.AddMethod();
+			MethodsBuilder.AddTask(TEXT("Task2"));
+			MethodsBuilder.AddTask(TEXT("Task1"));
+		}
+		FHTNBuilder_PrimitiveTask& PrimitiveTask1 = DomainBuilder.AddPrimitiveTask(TEXT("Task1"));
+		PrimitiveTask1.SetOperator(1, 2);
+		FHTNBuilder_PrimitiveTask& PrimitiveTask2 = DomainBuilder.AddPrimitiveTask(TEXT("Task2"));
+		PrimitiveTask2.SetOperator(3, 4);
+		
+		DomainBuilder.Compile();
+
+		FHTNResult Result;
+		Planner.GeneratePlan(*DomainBuilder.DomainInstance, WorldState, Result);
+		
+		Test(TEXT("Plan should contain two elements"), Result.TaskIDs.Num() == 2);
+		if (Result.TaskIDs.Num() >= 2)
+		{
+			const FHTNPolicy::FTaskID TaskID1 = DomainBuilder.DomainInstance->FindTaskID(TEXT("Task1"));
+			const FHTNPolicy::FTaskID TaskID2 = DomainBuilder.DomainInstance->FindTaskID(TEXT("Task2"));
+			Test(TEXT("Task2 should be the first one"), Result.TaskIDs[0] == TaskID2);
+			Test(TEXT("Task2 action should be the first one"), Result.ActionsSequence[0].ActionID == 3 && Result.ActionsSequence[0].Parameter == 4);
+			Test(TEXT("Task1 should be the second one"), Result.TaskIDs[1] == TaskID1);
+			Test(TEXT("Task1 action should be the first one"), Result.ActionsSequence[1].ActionID == 1 && Result.ActionsSequence[1].Parameter == 2);
+
+
+		}
+	}
+};
+IMPLEMENT_AI_LATENT_TEST(FAITest_HTNOperatorsOfGeneratedPlan, "System.Engine.AI.HTN.OperatorsOfGeneratedPlan")
+
+//----------------------------------------------------------------------//
+// Component tests
+//----------------------------------------------------------------------//
+typedef FAITest_SimpleComponentBasedTest<UMockHTNComponent> FAITest_HTNComponentTest;
+
+PRAGMA_ENABLE_OPTIMIZATION
+
+#undef LOCTEXT_NAMESPACE
diff --git a/Engine/Plugins/AI/HTNPlanner/Source/HTNTestSuite/Private/HTNTestSuiteModule.cpp b/Engine/Plugins/AI/HTNPlanner/Source/HTNTestSuite/Private/HTNTestSuiteModule.cpp
new file mode 100644
--- /dev/null
+++ b/Engine/Plugins/AI/HTNPlanner/Source/HTNTestSuite/Private/HTNTestSuiteModule.cpp
@@ -0,0 +1,13 @@
+// Copyright 1998-2018 Epic Games, Inc. All Rights Reserved.
+
+#include "HTNTestSuiteModule.h"
+
+#define LOCTEXT_NAMESPACE "HTNTestSuite"
+
+class FHTNTestSuiteModule : public IHTNTestSuiteModule
+{
+};
+
+IMPLEMENT_MODULE(FHTNTestSuiteModule, HTNTestSuiteModule)
+
+#undef LOCTEXT_NAMESPACE
diff --git a/Engine/Plugins/Editor/AssetManagerEditor/Source/AssetManagerEditor/Private/AssetManagerEditorModule.cpp b/Engine/Plugins/Editor/AssetManagerEditor/Source/AssetManagerEditor/Private/AssetManagerEditorModule.cpp
--- a/Engine/Plugins/Editor/AssetManagerEditor/Source/AssetManagerEditor/Private/AssetManagerEditorModule.cpp
+++ b/Engine/Plugins/Editor/AssetManagerEditor/Source/AssetManagerEditor/Private/AssetManagerEditorModule.cpp
@@ -969,6 +969,9 @@ bool FAssetManagerEditorModule::GetStringValueForCustomColumn(const FAssetData&
 		case EPrimaryAssetCookRule::AlwaysCook: 
 			OutValue = TEXT("Always");
 			return true;
+		case EPrimaryAssetCookRule::DevelopmentAlwaysCook:
+			OutValue = TEXT("DevelopmentAlways");
+			return true;
 		case EPrimaryAssetCookRule::DevelopmentCook: 
 			OutValue = TEXT("Development");
 			return true;
@@ -1055,6 +1058,9 @@ bool FAssetManagerEditorModule::GetDisplayTextForCustomColumn(const FAssetData&
 		case EPrimaryAssetCookRule::AlwaysCook:
 			OutValue = LOCTEXT("AlwaysCook", "Always");
 			return true;
+		case EPrimaryAssetCookRule::DevelopmentAlwaysCook:
+			OutValue = LOCTEXT("DevelopmentAlwaysCook", "DevelopmentAlways");
+			return true;
 		case EPrimaryAssetCookRule::DevelopmentCook:
 			OutValue = LOCTEXT("DevelopmentCook", "Development");
 			return true;
diff --git a/Engine/Plugins/Editor/AssetManagerEditor/Source/AssetManagerEditor/Private/SAssetAuditBrowser.cpp b/Engine/Plugins/Editor/AssetManagerEditor/Source/AssetManagerEditor/Private/SAssetAuditBrowser.cpp
--- a/Engine/Plugins/Editor/AssetManagerEditor/Source/AssetManagerEditor/Private/SAssetAuditBrowser.cpp
+++ b/Engine/Plugins/Editor/AssetManagerEditor/Source/AssetManagerEditor/Private/SAssetAuditBrowser.cpp
@@ -31,6 +31,7 @@
 #include "Widgets/Input/SComboBox.h"
 #include "Framework/Application/SlateApplication.h"
 #include "DragAndDrop/AssetDragDropOp.h"
+#include "Blueprint/BlueprintSupport.h"
 
 #define LOCTEXT_NAMESPACE "AssetManagementBrowser"
 
@@ -308,13 +309,13 @@ void SAssetAuditBrowser::Construct(const FArguments& InArgs)
 	AssetRegistryTagsToIgnore.Add(FPrimaryAssetId::PrimaryAssetNameTag);
 
 	// Ignore blueprint tags
-	AssetRegistryTagsToIgnore.Add("ParentClass");
-	AssetRegistryTagsToIgnore.Add("BlueprintType");
-	AssetRegistryTagsToIgnore.Add("NumReplicatedProperties");
-	AssetRegistryTagsToIgnore.Add("NativeParentClass");
-	AssetRegistryTagsToIgnore.Add("IsDataOnly");
-	AssetRegistryTagsToIgnore.Add("NativeComponents");
-	AssetRegistryTagsToIgnore.Add("BlueprintComponents");
+	AssetRegistryTagsToIgnore.Add(FBlueprintTags::ParentClassPath);
+	AssetRegistryTagsToIgnore.Add(FBlueprintTags::BlueprintType);
+	AssetRegistryTagsToIgnore.Add(FBlueprintTags::NumReplicatedProperties);
+	AssetRegistryTagsToIgnore.Add(FBlueprintTags::NativeParentClassPath);
+	AssetRegistryTagsToIgnore.Add(FBlueprintTags::IsDataOnly);
+	AssetRegistryTagsToIgnore.Add(FBlueprintTags::NumNativeComponents);
+	AssetRegistryTagsToIgnore.Add(FBlueprintTags::NumBlueprintComponents);
 
 	static const FName DefaultForegroundName("DefaultForeground");
 
diff --git a/Engine/Plugins/Runtime/GameplayAbilities/Source/GameplayAbilities/Private/GameplayCueManager.cpp b/Engine/Plugins/Runtime/GameplayAbilities/Source/GameplayAbilities/Private/GameplayCueManager.cpp
--- a/Engine/Plugins/Runtime/GameplayAbilities/Source/GameplayAbilities/Private/GameplayCueManager.cpp
+++ b/Engine/Plugins/Runtime/GameplayAbilities/Source/GameplayAbilities/Private/GameplayCueManager.cpp
@@ -856,7 +856,7 @@ void UGameplayCueManager::BuildCuesToAddToGlobalSet(const TArray<FAssetData>& As
 		
 		if (!FoundGameplayTag.IsNone())
 		{
-			const FString GeneratedClassTag = Data.GetTagValueRef<FString>("GeneratedClass");
+			const FString GeneratedClassTag = Data.GetTagValueRef<FString>(FBlueprintTags::GeneratedClassPath);
 			if (GeneratedClassTag.IsEmpty())
 			{
 				ABILITY_LOG(Warning, TEXT("Unable to find GeneratedClass value for AssetData %s"), *Data.ObjectPath.ToString());
@@ -1054,7 +1054,7 @@ void UGameplayCueManager::HandleAssetDeleted(UObject *Object)
 /** Handles cleaning up an object library if it matches the passed in object */
 void UGameplayCueManager::HandleAssetRenamed(const FAssetData& Data, const FString& String)
 {
-	const FString ParentClassName = Data.GetTagValueRef<FString>("ParentClass");
+	const FString ParentClassName = Data.GetTagValueRef<FString>(FBlueprintTags::ParentClassPath);
 	if (!ParentClassName.IsEmpty())
 	{
 		UClass* DataClass = FindObject<UClass>(nullptr, *ParentClassName);
diff --git a/Engine/Source/Developer/AssetTools/Private/AssetTypeActions/AssetTypeActions_Blueprint.cpp b/Engine/Source/Developer/AssetTools/Private/AssetTypeActions/AssetTypeActions_Blueprint.cpp
--- a/Engine/Source/Developer/AssetTools/Private/AssetTypeActions/AssetTypeActions_Blueprint.cpp
+++ b/Engine/Source/Developer/AssetTools/Private/AssetTypeActions/AssetTypeActions_Blueprint.cpp
@@ -298,15 +298,12 @@ FText FAssetTypeActions_Blueprint::GetAssetDescription(const FAssetData& AssetDa
 
 TWeakPtr<IClassTypeActions> FAssetTypeActions_Blueprint::GetClassTypeActions(const FAssetData& AssetData) const
 {
-	static const FName NativeParentClassTag = "NativeParentClass";
-	static const FName ParentClassTag = "ParentClass";
-
 	// Blueprints get the class type actions for their parent native class - this avoids us having to load the blueprint
 	UClass* ParentClass = nullptr;
 	FString ParentClassName;
-	if(!AssetData.GetTagValue(NativeParentClassTag, ParentClassName))
+	if(!AssetData.GetTagValue(FBlueprintTags::NativeParentClassPath, ParentClassName))
 	{
-		AssetData.GetTagValue(ParentClassTag, ParentClassName);
+		AssetData.GetTagValue(FBlueprintTags::ParentClassPath, ParentClassName);
 	}
 	if(!ParentClassName.IsEmpty())
 	{
diff --git a/Engine/Source/Developer/BlueprintCompilerCppBackend/Private/BlueprintCompilerCppBackendGatherDependencies.cpp b/Engine/Source/Developer/BlueprintCompilerCppBackend/Private/BlueprintCompilerCppBackendGatherDependencies.cpp
--- a/Engine/Source/Developer/BlueprintCompilerCppBackend/Private/BlueprintCompilerCppBackendGatherDependencies.cpp
+++ b/Engine/Source/Developer/BlueprintCompilerCppBackend/Private/BlueprintCompilerCppBackendGatherDependencies.cpp
@@ -163,11 +163,6 @@ struct FFindAssetsToInclude : public FGatherConvertedClassDependenciesHelperBase
 						// This is a field owned by a class, so attempt to add the class as a dependency.
 						MaybeIncludeObjectAsDependency(OwnerClass, CurrentlyConvertedStruct);
 					}
-					else if (ObjAsBPGC && (CurrentlyConvertedStruct->IsA<UUserDefinedStruct>() || ObjAsBPGC->IsChildOf(CurrentlyConvertedStruct)))
-					{
-						// If an unconverted BP (e.g. DOBP) has a potentially circular dependency on the converted asset, wrap it inside a stub class to avoid an EDL cycle on load.
-						Dependencies.MarkUnconvertedClassAsNecessary(ObjAsBPGC);
-					}
 					else
 					{
 						// Add the class itself as a dependency.
@@ -529,23 +524,6 @@ bool FGatherConvertedClassDependencies::WillClassBeConverted(const UBlueprintGen
 	return false;
 }
 
-void FGatherConvertedClassDependencies::MarkUnconvertedClassAsNecessary(const UBlueprintGeneratedClass* InClass) const
-{
-	if (InClass && !InClass->HasAnyFlags(RF_ClassDefaultObject))
-	{
-		const UClass* ClassToInclude = FindOriginalClass(InClass);
-		checkSlow(ClassToInclude);
-
-		IBlueprintCompilerCppBackendModule& BackEndModule = (IBlueprintCompilerCppBackendModule&)IBlueprintCompilerCppBackendModule::Get();
-		const auto& IncludeUnconvertedClassDelegate = BackEndModule.OnIncludingUnconvertedBP();
-
-		if (IncludeUnconvertedClassDelegate.IsBound())
-		{
-			IncludeUnconvertedClassDelegate.Execute(Cast<UBlueprint>(ClassToInclude->ClassGeneratedBy), NativizationOptions);
-		}
-	}
-}
-
 void FGatherConvertedClassDependencies::DependenciesForHeader()
 {
 	TArray<UObject*> ObjectsToCheck;
diff --git a/Engine/Source/Developer/BlueprintCompilerCppBackend/Private/BlueprintCompilerCppBackendUMG.cpp b/Engine/Source/Developer/BlueprintCompilerCppBackend/Private/BlueprintCompilerCppBackendUMG.cpp
--- a/Engine/Source/Developer/BlueprintCompilerCppBackend/Private/BlueprintCompilerCppBackendUMG.cpp
+++ b/Engine/Source/Developer/BlueprintCompilerCppBackend/Private/BlueprintCompilerCppBackendUMG.cpp
@@ -122,9 +122,18 @@ bool FBackendHelperUMG::SpecialStructureConstructorUMG(const UStruct* Struct, co
 		if (OutResult)
 		{
 			const FSectionEvaluationData* SectionEvaluationData = reinterpret_cast<const FSectionEvaluationData*>(ValuePtr);
-			*OutResult = FString::Printf(TEXT("FSectionEvaluationData(%d, %d)")
-				, SectionEvaluationData->ImplIndex
-				, SectionEvaluationData->ForcedTime.Value);
+			if (SectionEvaluationData->ForcedTime == TNumericLimits<int32>::Lowest())
+			{
+				*OutResult = FString::Printf(TEXT("FSectionEvaluationData(%d, ESectionEvaluationFlags(0x%02x))")
+					, SectionEvaluationData->ImplIndex
+					, (uint8)SectionEvaluationData->Flags);
+			}
+			else
+			{
+				*OutResult = FString::Printf(TEXT("FSectionEvaluationData(%d, %d)")
+					, SectionEvaluationData->ImplIndex
+					, SectionEvaluationData->ForcedTime.Value);
+			}
 		}
 		return true;
 	}
@@ -133,11 +142,23 @@ bool FBackendHelperUMG::SpecialStructureConstructorUMG(const UStruct* Struct, co
 	{
 		if (OutResult)
 		{
+			auto MovieSceneSegmentRangeBoundConstructorLambda = [](const TRangeBound<FFrameNumber>& InRangeBound, const FFrameNumber& InRangeBoundValue) -> FString
+			{
+				if (InRangeBound.IsExclusive())
+				{
+					return FString::Printf(TEXT("TRangeBound<FFrameNumber>::Exclusive(%d)"), InRangeBoundValue.Value);
+				}
+				else if (InRangeBound.IsInclusive())
+				{
+					return FString::Printf(TEXT("TRangeBound<FFrameNumber>::Inclusive(%d)"), InRangeBoundValue.Value);
+				}
+				else
+				{
+					return FString::Printf(TEXT("TRangeBound<FFrameNumber>::Open()"));
+				}
+			};
+
 			const FMovieSceneSegment* MovieSceneSegment = reinterpret_cast<const FMovieSceneSegment*>(ValuePtr);
-			FString RangeStr;
-			FEmitDefaultValueHelper::SpecialStructureConstructor(TBaseStructure<FFloatRange>::Get()
-				, reinterpret_cast<const uint8*>(&MovieSceneSegment->Range)
-				, &RangeStr);
 			FString SegmentsInitializerList;
 			for (const FSectionEvaluationData& SectionEvaluationData : MovieSceneSegment->Impls)
 			{
@@ -151,7 +172,11 @@ bool FBackendHelperUMG::SpecialStructureConstructorUMG(const UStruct* Struct, co
 					, &SectionEvaluationDataStr);
 				SegmentsInitializerList += SectionEvaluationDataStr;
 			}
-			*OutResult = FString::Printf(TEXT("FMovieSceneSegment(%s, {%s})"), *RangeStr, *SegmentsInitializerList);
+			const FString LowerBoundStr = MovieSceneSegmentRangeBoundConstructorLambda(MovieSceneSegment->Range.GetLowerBound()
+				, MovieSceneSegment->Range.GetLowerBound().IsClosed() ? MovieSceneSegment->Range.GetLowerBoundValue() : FFrameNumber());
+			const FString UpperBoundStr = MovieSceneSegmentRangeBoundConstructorLambda(MovieSceneSegment->Range.GetUpperBound()
+				, MovieSceneSegment->Range.GetUpperBound().IsClosed() ? MovieSceneSegment->Range.GetUpperBoundValue() : FFrameNumber());
+			*OutResult = FString::Printf(TEXT("FMovieSceneSegment(TRange<FFrameNumber>(%s, %s), {%s})"), *LowerBoundStr, *UpperBoundStr, *SegmentsInitializerList);
 		}
 		return true;
 	}
diff --git a/Engine/Source/Developer/BlueprintCompilerCppBackend/Private/BlueprintCompilerCppBackendValueHelper.cpp b/Engine/Source/Developer/BlueprintCompilerCppBackend/Private/BlueprintCompilerCppBackendValueHelper.cpp
--- a/Engine/Source/Developer/BlueprintCompilerCppBackend/Private/BlueprintCompilerCppBackendValueHelper.cpp
+++ b/Engine/Source/Developer/BlueprintCompilerCppBackend/Private/BlueprintCompilerCppBackendValueHelper.cpp
@@ -641,7 +641,7 @@ bool FEmitDefaultValueHelper::SpecialStructureConstructor(const UStruct* Struct,
 			}
 			if (RangeBound->IsInclusive())
 			{
-				*OutResult = FString::Printf(TEXT("FInt32RangeBound::%s(%d)"), GET_FUNCTION_NAME_STRING_CHECKED(FInt32RangeBound, Exclusive), RangeBound->GetValue());
+				*OutResult = FString::Printf(TEXT("FInt32RangeBound::%s(%d)"), GET_FUNCTION_NAME_STRING_CHECKED(FInt32RangeBound, Inclusive), RangeBound->GetValue());
 			}
 			if (RangeBound->IsOpen())
 			{
@@ -864,7 +864,7 @@ struct FDefaultSubobjectData
 
 protected:
 	// Generate special-case property initialization code. This could be something that is normally handled through custom serialization.
-	virtual bool HandledAsSpecialProperty(FEmitterLocalContext& Context, const UProperty* Property)
+	bool HandledAsSpecialProperty(FEmitterLocalContext& Context, const UProperty* Property)
 	{
 		bool bWasHandled = true;
 
@@ -922,11 +922,9 @@ struct FNonativeComponentData : public FDefaultSubobjectData
 	FString ParentVariableName;
 	/** Socket/Bone that Component might attach to */
 	FName AttachToName;
-	bool bIsRoot;
 
 	FNonativeComponentData()
 		: SCSNode(nullptr)
-		, bIsRoot(false)
 	{
 		bAddLocalScope = false;
 	}
@@ -956,24 +954,6 @@ struct FNonativeComponentData : public FDefaultSubobjectData
 		// Continue inline here with the default logic, but we don't need to enclose it within a new scope block.
 		FDefaultSubobjectData::EmitPropertyInitialization(Context);
 	}
-
-protected:
-	virtual bool HandledAsSpecialProperty(FEmitterLocalContext& Context, const UProperty* Property) override
-	{
-		// skip relative location and rotation. THey are ignored for root components created from scs (and they probably should be reset by scs editor).
-		if (bIsRoot && (Property->GetOuter() == USceneComponent::StaticClass()))
-		{
-			UProperty* RelativeLocationProperty = USceneComponent::StaticClass()->FindPropertyByName(GET_MEMBER_NAME_CHECKED(USceneComponent, RelativeLocation));
-			UProperty* RelativeRotationProperty = USceneComponent::StaticClass()->FindPropertyByName(GET_MEMBER_NAME_CHECKED(USceneComponent, RelativeRotation));
-			if ((Property == RelativeLocationProperty) || (Property == RelativeRotationProperty))
-			{
-				return true;
-			}
-		}
-
-		// Continue on with default logic.
-		return FDefaultSubobjectData::HandledAsSpecialProperty(Context, Property);
-	}
 };
 
 FString FEmitDefaultValueHelper::HandleNonNativeComponent(FEmitterLocalContext& Context, const USCS_Node* Node
@@ -1015,9 +995,6 @@ FString FEmitDefaultValueHelper::HandleNonNativeComponent(FEmitterLocalContext&
 			NonativeComponentData.SCSNode = Node;
 			NonativeComponentData.VariableName = NativeVariablePropertyName;
 			NonativeComponentData.Object = ComponentTemplate;
-			USCS_Node* RootComponentNode = nullptr;
-			Node->GetSCS()->GetSceneRootComponentTemplate(&RootComponentNode);
-			NonativeComponentData.bIsRoot = RootComponentNode == Node;
 			UClass* ComponentClass = ComponentTemplate->GetClass();
 			check(ComponentClass != nullptr);
 
diff --git a/Engine/Source/Developer/BlueprintCompilerCppBackend/Public/BlueprintCompilerCppBackendGatherDependencies.h b/Engine/Source/Developer/BlueprintCompilerCppBackend/Public/BlueprintCompilerCppBackendGatherDependencies.h
--- a/Engine/Source/Developer/BlueprintCompilerCppBackend/Public/BlueprintCompilerCppBackendGatherDependencies.h
+++ b/Engine/Source/Developer/BlueprintCompilerCppBackend/Public/BlueprintCompilerCppBackendGatherDependencies.h
@@ -49,7 +49,6 @@ public:
 
 public:
 	bool WillClassBeConverted(const UBlueprintGeneratedClass* InClass) const;
-	void MarkUnconvertedClassAsNecessary(const UBlueprintGeneratedClass* InClass) const;
 
 	static void GatherAssetsReferencedByUDSDefaultValue(TSet<UObject*>& Dependencies, UUserDefinedStruct* Struct);
 
diff --git a/Engine/Source/Developer/BlueprintNativeCodeGen/Private/BlueprintNativeCodeGenModule.cpp b/Engine/Source/Developer/BlueprintNativeCodeGen/Private/BlueprintNativeCodeGenModule.cpp
--- a/Engine/Source/Developer/BlueprintNativeCodeGen/Private/BlueprintNativeCodeGenModule.cpp
+++ b/Engine/Source/Developer/BlueprintNativeCodeGen/Private/BlueprintNativeCodeGenModule.cpp
@@ -79,6 +79,10 @@ private:
 	void GenerateSingleStub(UBlueprint* BP, const FName PlatformName);
 	void CollectBoundFunctions(UBlueprint* BP);
 	void GenerateSingleAsset(UField* ForConversion, const FName PlatformName, TSharedPtr<FNativizationSummary> NativizationSummary = TSharedPtr<FNativizationSummary>());
+	void ReplaceAsset(const UObject* InAsset, const FCompilerNativizationOptions& NativizationOptions) const;
+	void GatherClassAssetsReferencedByStruct(TSet<const UBlueprintGeneratedClass*>& Assets, const UStruct* OuterStruct, const UStruct* InnerStruct = nullptr) const;
+	void ReplaceAssetsWithCircularReferenceTo(const UBlueprintGeneratedClass* InClass, const FCompilerNativizationOptions& NativizationOptions) const;
+	bool HasCircularReferenceWithAnyConvertedAsset(const UBlueprintGeneratedClass* InClass, const FCompilerNativizationOptions& NativizationOptions) const;
 
 	struct FStatePerPlatform
 	{
@@ -890,13 +894,8 @@ EReplacementResult FBlueprintNativeCodeGenModule::IsTargetedForReplacement(const
 		}
 		return false;
 	};
-	if (ObjectIsNotReplacedAtAll())
-	{
-		StateForCurrentPlatform->CachedIsTargetedForReplacement.Add(ObjectKey, EReplacementResult::DontReplace);
-		return EReplacementResult::DontReplace;
-	}
 
-	auto ObjectGenratesOnlyStub = [&]() -> bool
+	auto ObjectGeneratesOnlyStub = [&]() -> bool
 	{
 		// ExcludedFolderPaths
 		{
@@ -1050,14 +1049,165 @@ EReplacementResult FBlueprintNativeCodeGenModule::IsTargetedForReplacement(const
 		}
 		return false;
 	};
-	if (ObjectGenratesOnlyStub())
+
+	EReplacementResult Result = EReplacementResult::ReplaceCompletely;
+
+	if (ObjectIsNotReplacedAtAll())
+	{
+		Result = EReplacementResult::DontReplace;
+	}
+	else if (ObjectGeneratesOnlyStub())
+	{
+		Result = EReplacementResult::GenerateStub;
+	}
+
+	StateForCurrentPlatform->CachedIsTargetedForReplacement.Add(ObjectKey, Result);
+
+	if (BlueprintClass)
 	{
-		StateForCurrentPlatform->CachedIsTargetedForReplacement.Add(ObjectKey, EReplacementResult::GenerateStub);
-		return EReplacementResult::GenerateStub;
+		if (Result == EReplacementResult::ReplaceCompletely)
+		{
+			// Look for any circular references with unconverted assets. We'll need to convert those as well in order to avoid creating an EDL cycle.
+			ReplaceAssetsWithCircularReferenceTo(BlueprintClass, NativizationOptions);
+		}
+		else if(HasCircularReferenceWithAnyConvertedAsset(BlueprintClass, NativizationOptions))
+		{
+			UE_LOG(LogBlueprintCodeGen, Log, TEXT("Forcing '%s' to be replaced as it has a circular reference to a converted asset"), *BlueprintClass->GetName());
+
+			// Force unconverted assets to be replaced if it has a circular reference with any converted asset.
+			ReplaceAsset(BlueprintClass, NativizationOptions);
+		}
+	}
+
+	return Result;
+}
+
+void FBlueprintNativeCodeGenModule::ReplaceAsset(const UObject* InAsset, const FCompilerNativizationOptions& NativizationOptions) const
+{
+	if (InAsset)
+	{
+		const FStatePerPlatform* StateForCurrentPlatform = StatesPerPlatform.Find(NativizationOptions.PlatformName);
+		check(StateForCurrentPlatform);
+
+		const FSoftObjectPath ObjectKey(InAsset);
+		StateForCurrentPlatform->CachedIsTargetedForReplacement.FindChecked(ObjectKey) = EReplacementResult::ReplaceCompletely;
+	}
+}
+
+void FBlueprintNativeCodeGenModule::GatherClassAssetsReferencedByStruct(TSet<const UBlueprintGeneratedClass*>& Assets, const UStruct* OuterStruct, const UStruct* InnerStruct) const
+{
+	if (OuterStruct)
+	{
+		for (const UProperty* Property = (InnerStruct ? InnerStruct->PropertyLink : OuterStruct->PropertyLink); Property; Property = Property->PropertyLinkNext)
+		{
+			TArray<const UProperty*, TInlineAllocator<2>> InnerPropertyList;
+			if (const UArrayProperty* ArrayProperty = Cast<const UArrayProperty>(Property))
+			{
+				InnerPropertyList.Add(ArrayProperty->Inner);
+			}
+			else if (const USetProperty* SetProperty = Cast<const USetProperty>(Property))
+			{
+				InnerPropertyList.Add(SetProperty->ElementProp);
+			}
+			else if (const UMapProperty* MapProperty = Cast<const UMapProperty>(Property))
+			{
+				InnerPropertyList.Add(MapProperty->KeyProp);
+				InnerPropertyList.Add(MapProperty->ValueProp);
+			}
+			else
+			{
+				InnerPropertyList.Add(Property);
+			}
+
+			for (const UProperty* InnerProperty : InnerPropertyList)
+			{
+				if (const UStructProperty* StructProperty = Cast<UStructProperty>(InnerProperty))
+				{
+					GatherClassAssetsReferencedByStruct(Assets, OuterStruct, StructProperty->Struct);
+				}
+				else
+				{
+					const UBlueprintGeneratedClass* BPGC = nullptr;
+					if (const UObjectPropertyBase* ObjectProperty = Cast<UObjectPropertyBase>(InnerProperty))
+					{
+						BPGC = Cast<const UBlueprintGeneratedClass>(ObjectProperty->PropertyClass);
+					}
+					else if (const UClassProperty* ClassProperty = Cast<UClassProperty>(InnerProperty))
+					{
+						BPGC = Cast<const UBlueprintGeneratedClass>(ClassProperty->MetaClass);
+					}
+
+					if (BPGC)
+					{
+						Assets.Add(BPGC);
+					}
+				}
+			}
+		}
+	}
+}
+
+void FBlueprintNativeCodeGenModule::ReplaceAssetsWithCircularReferenceTo(const UBlueprintGeneratedClass* InClass, const FCompilerNativizationOptions& NativizationOptions) const
+{
+	TSet<const UBlueprintGeneratedClass*> ForwardReferencedAssets;
+	GatherClassAssetsReferencedByStruct(ForwardReferencedAssets, InClass);
+
+	for (const UBlueprintGeneratedClass* ForwardReference : ForwardReferencedAssets)
+	{
+		const EReplacementResult Result = IsTargetedForReplacement(ForwardReference, NativizationOptions);
+		if (Result != EReplacementResult::ReplaceCompletely)
+		{
+			bool bForceConvert = false;
+			if (ForwardReference->IsChildOf(InClass))
+			{
+				UE_LOG(LogBlueprintCodeGen, Log, TEXT("Forcing '%s' to be replaced as it has a circular reference with '%s'"), *ForwardReference->GetName(), *InClass->GetName());
+				
+				ReplaceAsset(ForwardReference, NativizationOptions);
+			}
+			else
+			{
+				TSet<const UBlueprintGeneratedClass*> ReverseReferencedAssets;
+				GatherClassAssetsReferencedByStruct(ReverseReferencedAssets, ForwardReference);
+
+				for (const UBlueprintGeneratedClass* ReverseReference : ReverseReferencedAssets)
+				{
+					if (ReverseReference->IsChildOf(InClass))
+					{
+						UE_LOG(LogBlueprintCodeGen, Log, TEXT("Forcing '%s' to be replaced as it has a circular reference to '%s'"), *ForwardReference->GetName(), *InClass->GetName());
+						
+						ReplaceAsset(ForwardReference, NativizationOptions);
+						break;
+					}
+				}
+			}
+		}
+	}
+}
+
+bool FBlueprintNativeCodeGenModule::HasCircularReferenceWithAnyConvertedAsset(const UBlueprintGeneratedClass* InClass, const FCompilerNativizationOptions& NativizationOptions) const
+{
+	TSet<const UBlueprintGeneratedClass*> ForwardReferencedAssets;
+	GatherClassAssetsReferencedByStruct(ForwardReferencedAssets, InClass);
+
+	for (const UBlueprintGeneratedClass* ForwardReference : ForwardReferencedAssets)
+	{
+		const EReplacementResult Result = IsTargetedForReplacement(ForwardReference, NativizationOptions);
+		if (Result == EReplacementResult::ReplaceCompletely)
+		{
+			TSet<const UBlueprintGeneratedClass*> ReverseReferencedAssets;
+			GatherClassAssetsReferencedByStruct(ReverseReferencedAssets, ForwardReference);
+
+			for (const UBlueprintGeneratedClass* ReverseReference : ReverseReferencedAssets)
+			{
+				if (ReverseReference->IsChildOf(InClass))
+				{
+					return true;
+				}
+			}
+		}
 	}
 
-	StateForCurrentPlatform->CachedIsTargetedForReplacement.Add(ObjectKey, EReplacementResult::ReplaceCompletely);
-	return EReplacementResult::ReplaceCompletely;
+	return false;
 }
 
 void FBlueprintNativeCodeGenModule::FillPlatformNativizationDetails(const ITargetPlatform* Platform, FPlatformNativizationDetails& Details)
diff --git a/Engine/Source/Developer/OutputLog/Private/SOutputLog.cpp b/Engine/Source/Developer/OutputLog/Private/SOutputLog.cpp
--- a/Engine/Source/Developer/OutputLog/Private/SOutputLog.cpp
+++ b/Engine/Source/Developer/OutputLog/Private/SOutputLog.cpp
@@ -705,7 +705,10 @@ FReply SConsoleInputBox::OnKeyCharHandler(const FGeometry& MyGeometry, const FCh
 	const uint32* KeyCode = nullptr;
 	const uint32* CharCode = nullptr;
 	FInputKeyManager::Get().GetCodesFromKey(OpenConsoleChord.Key, KeyCode, CharCode);
-	check(CharCode);
+	if (CharCode == nullptr)
+	{
+		return FReply::Unhandled();
+	}
 
 	// Intercept the "open console" key
 	if ((!ActiveCommandExecutor || ActiveCommandExecutor->AllowHotKeyClose())
diff --git a/Engine/Source/Editor/AIGraph/Private/AIGraphTypes.cpp b/Engine/Source/Editor/AIGraph/Private/AIGraphTypes.cpp
--- a/Engine/Source/Editor/AIGraph/Private/AIGraphTypes.cpp
+++ b/Engine/Source/Editor/AIGraph/Private/AIGraphTypes.cpp
@@ -301,7 +301,7 @@ void FGraphNodeClassHelper::OnAssetAdded(const struct FAssetData& AssetData)
 void FGraphNodeClassHelper::OnAssetRemoved(const struct FAssetData& AssetData)
 {
 	FString AssetClassName;
-	if (AssetData.GetTagValue("GeneratedClass", AssetClassName))
+	if (AssetData.GetTagValue(FBlueprintTags::GeneratedClassPath, AssetClassName))
 	{
 		ConstructorHelpers::StripObjectClass(AssetClassName);
 		AssetClassName = FPackageName::ObjectPathToObjectName(*AssetClassName);
@@ -338,7 +338,7 @@ TSharedPtr<FGraphNodeClassNode> FGraphNodeClassHelper::CreateClassDataNode(const
 
 	FString AssetClassName;
 	FString AssetParentClassName;
-	if (AssetData.GetTagValue("GeneratedClass", AssetClassName) && AssetData.GetTagValue("ParentClass", AssetParentClassName))
+	if (AssetData.GetTagValue(FBlueprintTags::GeneratedClassPath, AssetClassName) && AssetData.GetTagValue(FBlueprintTags::ParentClassPath, AssetParentClassName))
 	{
 		UObject* Outer1(NULL);
 		ResolveName(Outer1, AssetClassName, false, false);
diff --git a/Engine/Source/Editor/BehaviorTreeEditor/Classes/BehaviorTreeGraphNode_Root.h b/Engine/Source/Editor/BehaviorTreeEditor/Classes/BehaviorTreeGraphNode_Root.h
--- a/Engine/Source/Editor/BehaviorTreeEditor/Classes/BehaviorTreeGraphNode_Root.h
+++ b/Engine/Source/Editor/BehaviorTreeEditor/Classes/BehaviorTreeGraphNode_Root.h
@@ -13,7 +13,7 @@ class UBehaviorTreeGraphNode_Root : public UBehaviorTreeGraphNode
 {
 	GENERATED_UCLASS_BODY()
 
-	UPROPERTY(EditAnywhere, Category=BehaviorTree)
+	UPROPERTY(EditAnywhere, Category="AI|BehaviorTree")
 	class UBlackboardData* BlackboardAsset;
 
 	virtual void PostPlacedNewNode() override;
diff --git a/Engine/Source/Editor/BlueprintGraph/Classes/EdGraphSchema_K2.h b/Engine/Source/Editor/BlueprintGraph/Classes/EdGraphSchema_K2.h
--- a/Engine/Source/Editor/BlueprintGraph/Classes/EdGraphSchema_K2.h
+++ b/Engine/Source/Editor/BlueprintGraph/Classes/EdGraphSchema_K2.h
@@ -574,7 +574,7 @@ public:
 	virtual bool DefaultValueSimpleValidation(const FEdGraphPinType& PinType, const FName PinName, const FString& NewDefaultValue, UObject* NewDefaultObject, const FText& InText, FString* OutMsg = nullptr) const;
 
 	/** Returns true if the owning node is a function with AutoCreateRefTerm meta data */
-	bool IsAutoCreateRefTerm(const UEdGraphPin* Pin) const;
+	static bool IsAutoCreateRefTerm(const UEdGraphPin* Pin);
 
 	/** See if a class has any members that are accessible by a blueprint */
 	bool ClassHasBlueprintAccessibleMembers(const UClass* InClass) const;
@@ -808,7 +808,7 @@ public:
 	 * @param	Function				The function to check for auto-emitted ref terms on
 	 * @param	AutoEmitParameterNames	(out) Returns an array of param names that should be auto-emitted if nothing is connected
 	 */
-	void GetAutoEmitTermParameters(const UFunction* Function, TArray<FString>& AutoEmitParameterNames) const;
+	static void GetAutoEmitTermParameters(const UFunction* Function, TArray<FString>& AutoEmitParameterNames);
 
 	/**
 	 * Determine if a function has a parameter of a specific type.
diff --git a/Engine/Source/Editor/BlueprintGraph/Private/EdGraphSchema_K2.cpp b/Engine/Source/Editor/BlueprintGraph/Private/EdGraphSchema_K2.cpp
--- a/Engine/Source/Editor/BlueprintGraph/Private/EdGraphSchema_K2.cpp
+++ b/Engine/Source/Editor/BlueprintGraph/Private/EdGraphSchema_K2.cpp
@@ -380,13 +380,13 @@ public:
 
 					if (Asset.IsValid() && !Asset.IsAssetLoaded())
 					{
-						const FString BlueprintTypeStr = Asset.GetTagValueRef<FString>("BlueprintType");
+						const FString BlueprintTypeStr = Asset.GetTagValueRef<FString>(FBlueprintTags::BlueprintType);
 						const bool bNormalBP = BlueprintTypeStr == BPNormalTypeAllowed;
 						const bool bInterfaceBP = BlueprintTypeStr == BPInterfaceTypeAllowed;
 
 						if (bNormalBP || bInterfaceBP)
 						{
-							const uint32 ClassFlags = Asset.GetTagValueRef<uint32>("ClassFlags");
+							const uint32 ClassFlags = Asset.GetTagValueRef<uint32>(FBlueprintTags::ClassFlags);
 							if (!(ClassFlags & CLASS_Deprecated))
 							{
 								if (bNormalBP)
@@ -956,13 +956,13 @@ FText UEdGraphSchema_K2::GetFriendlySignatureName(const UFunction* Function)
 	return UK2Node_CallFunction::GetUserFacingFunctionName( Function );
 }
 
-void UEdGraphSchema_K2::GetAutoEmitTermParameters(const UFunction* Function, TArray<FString>& AutoEmitParameterNames) const
+void UEdGraphSchema_K2::GetAutoEmitTermParameters(const UFunction* Function, TArray<FString>& AutoEmitParameterNames)
 {
-	AutoEmitParameterNames.Empty();
+	AutoEmitParameterNames.Reset();
 
-	if( Function->HasMetaData(FBlueprintMetadata::MD_AutoCreateRefTerm) )
+	const FString& MetaData = Function->GetMetaData(FBlueprintMetadata::MD_AutoCreateRefTerm);
+	if (!MetaData.IsEmpty())
 	{
-		const FString& MetaData = Function->GetMetaData(FBlueprintMetadata::MD_AutoCreateRefTerm);
 		MetaData.ParseIntoArray(AutoEmitParameterNames, TEXT(","), true);
 
 		for (int32 NameIndex = 0; NameIndex < AutoEmitParameterNames.Num();)
@@ -2820,7 +2820,7 @@ FString UEdGraphSchema_K2::IsPinDefaultValid(const UEdGraphPin* Pin, const FStri
 	{
 		if( !bIsAutoCreateRefTerm )
 		{
-			// No harm in leaving a function result node input (aka functin output) unconnected - the property will be initialized correctly
+			// No harm in leaving a function result node input (aka function output) unconnected - the property will be initialized correctly
 			// as empty:
 			bool bIsFunctionOutput = false;
 			if(UK2Node_FunctionResult* Result = Cast<UK2Node_FunctionResult>(Pin->GetOwningNode()))
@@ -4699,21 +4699,21 @@ bool UEdGraphSchema_K2::DoesDefaultValueMatchAutogenerated(const UEdGraphPin& In
 	return Super::DoesDefaultValueMatchAutogenerated(InPin);
 }
 
-bool UEdGraphSchema_K2::IsAutoCreateRefTerm(const UEdGraphPin* Pin) const
+bool UEdGraphSchema_K2::IsAutoCreateRefTerm(const UEdGraphPin* Pin)
 {
 	check(Pin);
 
 	bool bIsAutoCreateRefTerm = false;
-	UEdGraphNode* OwningNode = Pin->GetOwningNode();
-	UK2Node_CallFunction* FuncNode = Cast<UK2Node_CallFunction>(OwningNode);
-	if (FuncNode)
+	if (!Pin->PinName.IsNone())
 	{
-		UFunction* TargetFunction = FuncNode->GetTargetFunction();
-		if (TargetFunction && !Pin->PinName.IsNone())
+		if (UK2Node_CallFunction* FuncNode = Cast<UK2Node_CallFunction>(Pin->GetOwningNode()))
 		{
-			TArray<FString> AutoCreateParameterNames;
-			GetAutoEmitTermParameters(TargetFunction, AutoCreateParameterNames);
-			bIsAutoCreateRefTerm = AutoCreateParameterNames.Contains(Pin->PinName.ToString());
+			if (UFunction* TargetFunction = FuncNode->GetTargetFunction())
+			{
+				static TArray<FString> AutoCreateParameterNames;
+				GetAutoEmitTermParameters(TargetFunction, AutoCreateParameterNames);
+				bIsAutoCreateRefTerm = AutoCreateParameterNames.Contains(Pin->PinName.ToString());
+			}
 		}
 	}
 
diff --git a/Engine/Source/Editor/Blutility/Private/BlutilityMenuExtensions.cpp b/Engine/Source/Editor/Blutility/Private/BlutilityMenuExtensions.cpp
--- a/Engine/Source/Editor/Blutility/Private/BlutilityMenuExtensions.cpp
+++ b/Engine/Source/Editor/Blutility/Private/BlutilityMenuExtensions.cpp
@@ -165,7 +165,7 @@ void FBlutilityMenuExtensions::GetBlutilityClasses(TArray<FAssetData>& OutAssets
 	// Check each asset to see if it matches our type
 	for (const FAssetData& Asset : AssetList)
 	{
-		if (const FString* GeneratedClassPathPtr = Asset.TagsAndValues.Find(TEXT("GeneratedClass")))
+		if (const FString* GeneratedClassPathPtr = Asset.TagsAndValues.Find(FBlueprintTags::GeneratedClassPath))
 		{
 			const FString ClassObjectPath = FPackageName::ExportTextPathToObjectPath(*GeneratedClassPathPtr);
 			const FString ClassName = FPackageName::ObjectPathToObjectName(ClassObjectPath);
diff --git a/Engine/Source/Editor/ClassViewer/Private/ClassViewerNode.cpp b/Engine/Source/Editor/ClassViewer/Private/ClassViewerNode.cpp
--- a/Engine/Source/Editor/ClassViewer/Private/ClassViewerNode.cpp
+++ b/Engine/Source/Editor/ClassViewer/Private/ClassViewerNode.cpp
@@ -30,11 +30,10 @@ FClassViewerNode::FClassViewerNode( const FClassViewerNode& InCopyObject)
 	
 	UnloadedBlueprintData = InCopyObject.UnloadedBlueprintData;
 
-	GeneratedClassPackage = InCopyObject.GeneratedClassPackage;
-	GeneratedClassname = InCopyObject.GeneratedClassname;
-	ParentClassname = InCopyObject.ParentClassname;
+	ClassPath = InCopyObject.ClassPath;
+	ParentClassPath = InCopyObject.ParentClassPath;
 	ClassName = InCopyObject.ClassName;
-	AssetName = InCopyObject.AssetName;
+	BlueprintAssetPath = InCopyObject.BlueprintAssetPath;
 	bIsBPNormalType = InCopyObject.bIsBPNormalType;
 
 	// We do not want to copy the child list, do not add it. It should be the only item missing.
@@ -109,3 +108,8 @@ bool FClassViewerNode::IsClassPlaceable() const
 
 	return false;
 }
+
+bool FClassViewerNode::IsBlueprintClass() const
+{
+	return BlueprintAssetPath != NAME_None;
+}
diff --git a/Engine/Source/Editor/ClassViewer/Private/ClassViewerNode.h b/Engine/Source/Editor/ClassViewer/Private/ClassViewerNode.h
--- a/Engine/Source/Editor/ClassViewer/Private/ClassViewerNode.h
+++ b/Engine/Source/Editor/ClassViewer/Private/ClassViewerNode.h
@@ -39,7 +39,7 @@ public:
 	void AddUniqueChild(TSharedPtr<FClassViewerNode> NewChild);
 
 	/** 
-	 * Retrieves the class name this node is associated with. 
+	 * Retrieves the class name this node is associated with. This is not the literal UClass name as it is missing the _C for blueprints
 	 * @param	bUseDisplayName	Whether to use the display name or class name
 	 */
 	TSharedPtr<FString> GetClassName(bool bUseDisplayName = false) const
@@ -56,13 +56,17 @@ public:
 	/** Checks if the class is placeable. */
 	bool IsClassPlaceable() const;
 
+	/** Checks if this is a blueprint */
+	bool IsBlueprintClass() const;
+
+	/** Rather this class is not allowed for the specific context */
 	bool IsRestricted() const;
 
 private:
-	/** The class name for this tree node. */
+	/** The nontranslated internal name for this class. This is not necessarily the UClass's name, as that may have _C for blueprints */
 	TSharedPtr<FString> ClassName;
 
-	/** The class display name for this tree node. */
+	/** The translated display name for this class */
 	TSharedPtr<FString> ClassDisplayName;
 
 	/** List of children. */
@@ -75,17 +79,14 @@ public:
 	/** The blueprint this node is associated with. */
 	TWeakObjectPtr<UBlueprint> Blueprint;
 
-	/** Used to load up the package if it is unloaded, retrieved from meta data for the package. */
-	FString GeneratedClassPackage;
-
-	/** Used to examine the classname, retrieved from meta data for the package. */
-	FName GeneratedClassname;
+	/** Full object path to the class including _C, set for both blueprint and native */
+	FName ClassPath;
 
-	/** Used to find the parent of this class, retrieved from meta data for the package. */
-	FName ParentClassname;
+	/** Full object path to the parent class, may be blueprint or native */
+	FName ParentClassPath;
 
-	/** Used to load up the class if it is unloaded. */
-	FString AssetName;
+	/** Full path to the Blueprint that this class is loaded from, none for native classes*/
+	FName BlueprintAssetPath;
 
 	/** true if the class passed the filter. */
 	bool bPassesFilter;
diff --git a/Engine/Source/Editor/ClassViewer/Private/SClassViewer.cpp b/Engine/Source/Editor/ClassViewer/Private/SClassViewer.cpp
--- a/Engine/Source/Editor/ClassViewer/Private/SClassViewer.cpp
+++ b/Engine/Source/Editor/ClassViewer/Private/SClassViewer.cpp
@@ -414,10 +414,10 @@ public:
 	TSharedPtr< FClassViewerNode > FindParent(const TSharedPtr< FClassViewerNode >& InRootNode, FName InParentClassname, const UClass* InParentClass);
 
 	/** Updates the Class of a node. Uses the generated class package name to find the node.
-	 *	@param InGeneratedClassPackageName			The name of the generated class package to find the node for.
-	 *	@param InNewClass							The class to update the node with.
+	 *	@param InGeneratedClassPath		The path of the generated class to find the node for.
+	 *	@param InNewClass				The class to update the node with.
 	*/
-	void UpdateClassInNode(const FString& InGeneratedClassPackageName, UClass* InNewClass, UBlueprint* InNewBluePrint );
+	void UpdateClassInNode(FName InGeneratedClassPath, UClass* InNewClass, UBlueprint* InNewBluePrint );
 
 	/** Finds the node, recursively going deeper into the hierarchy. Does so by comparing class names.
 	*	@param InClassName			The name of the generated class package to find the node for.
@@ -437,11 +437,11 @@ private:
 	void OnHotReload( bool bWasTriggeredAutomatically );
 
 	/** Finds the node, recursively going deeper into the hierarchy. Does so by comparing generated class package names.
-	 *	@param InGeneratedClassPackageName			The name of the generated class package to find the node for.
+	 *	@param InGeneratedClassPath		The path of the generated class to find the node for.
 	 *
 	 *	@return The node.
 	 */
-	TSharedPtr< FClassViewerNode > FindNodeByGeneratedClassPackageName(const TSharedPtr< FClassViewerNode >& InRootNode, const FString& InGeneratedClassPackageName);
+	TSharedPtr< FClassViewerNode > FindNodeByGeneratedClassPath(const TSharedPtr< FClassViewerNode >& InRootNode, FName InGeneratedClassPath);
 
 	/** 
 	 * Loads the tag data for an unloaded blueprint asset.
@@ -461,12 +461,12 @@ private:
 	/** 
 	 * Recursively searches through the hierarchy to find and remove the asset. Used when deleting assets.
 	 *
-	 * @param InRootNode		The node to start the search with.
-	 * @param InAssetPackage	The package name of the asset to delete.
+	 * @param InRootNode	The node to start the search with.
+	 * @param InClassPath	The class path of the asset to delete
 	 *
 	 * @return Returns true if the asset was found and deleted successfully.
 	 */
-	bool FindAndRemoveNodeByPackageName(const TSharedPtr< FClassViewerNode >& InRootNode, const FString& InAssetPackage);
+	bool FindAndRemoveNodeByClassPath(const TSharedPtr< FClassViewerNode >& InRootNode, FName InClassPath);
 
 	/** Callback registered to the Asset Registry to be notified when an asset is added. */
 	void AddAsset(const FAssetData& InAddedAssetData);
@@ -501,7 +501,7 @@ namespace ClassViewer
 		// Pre-declare these functions.
 		static bool CheckIfBlueprintBase( TSharedPtr< FClassViewerNode> InNode );
 		static UBlueprint* GetBlueprint( UClass* InClass );
-		static void UpdateClassInNode(const FString& InGeneratedClassPackageName, UClass* InNewClass, UBlueprint* InNewBluePrint );
+		static void UpdateClassInNode(FName InGeneratedClassPath, UClass* InNewClass, UBlueprint* InNewBluePrint );
 
 		/** Util class to checks if a particular class can be made into a Blueprint, ignores deprecation
 		 *
@@ -696,14 +696,16 @@ namespace ClassViewer
 			bool bPassesDeveloperFilter = true;
 			static const FString DeveloperPathWithSlash = FPackageName::FilenameToLongPackageName(FPaths::GameDevelopersDir());
 			static const FString UserDeveloperPathWithSlash = FPackageName::FilenameToLongPackageName(FPaths::GameUserDeveloperDir());
+			FString GeneratedClassPathString = InOriginalRootNode->ClassPath.ToString();
+
 			if (InAllowedDeveloperType == EClassViewerDeveloperType::CVDT_None)
 			{
-				bPassesDeveloperFilter = !InOriginalRootNode->GeneratedClassPackage.StartsWith(DeveloperPathWithSlash);
+				bPassesDeveloperFilter = !GeneratedClassPathString.StartsWith(DeveloperPathWithSlash);
 			}
 			else if (InAllowedDeveloperType == EClassViewerDeveloperType::CVDT_CurrentUser)
 			{
-				if (InOriginalRootNode->GeneratedClassPackage.StartsWith(DeveloperPathWithSlash))
-					bPassesDeveloperFilter = InOriginalRootNode->GeneratedClassPackage.StartsWith(UserDeveloperPathWithSlash);
+				if (GeneratedClassPathString.StartsWith(DeveloperPathWithSlash))
+					bPassesDeveloperFilter = GeneratedClassPathString.StartsWith(UserDeveloperPathWithSlash);
 			}
 
 			// The INI files declare classes and folders that are considered internal only. Does this class match any of those patterns?
@@ -713,7 +715,7 @@ namespace ClassViewer
 			{
 				for (int i = 0; i < InternalPaths.Num(); i++)
 				{
-					if (InOriginalRootNode->GeneratedClassPackage.StartsWith(InternalPaths[i].Path))
+					if (GeneratedClassPathString.StartsWith(InternalPaths[i].Path))
 					{
 						bPassesInternalFilter = false;
 						break;
@@ -848,14 +850,15 @@ namespace ClassViewer
 			bool bPassesDeveloperFilter = true;
 			static const FString DeveloperPathWithSlash = FPackageName::FilenameToLongPackageName(FPaths::GameDevelopersDir());
 			static const FString UserDeveloperPathWithSlash = FPackageName::FilenameToLongPackageName(FPaths::GameUserDeveloperDir());
+			FString GeneratedClassPathString = InOriginalRootNode->ClassPath.ToString();
 			if (InAllowedDeveloperType == EClassViewerDeveloperType::CVDT_None)
 			{
-				bPassesDeveloperFilter = !InOriginalRootNode->GeneratedClassPackage.StartsWith(DeveloperPathWithSlash);
+				bPassesDeveloperFilter = GeneratedClassPathString.StartsWith(DeveloperPathWithSlash);
 			}
 			else if (InAllowedDeveloperType == EClassViewerDeveloperType::CVDT_CurrentUser)
 			{
-				if (InOriginalRootNode->GeneratedClassPackage.StartsWith(DeveloperPathWithSlash))
-					bPassesDeveloperFilter = InOriginalRootNode->GeneratedClassPackage.StartsWith(UserDeveloperPathWithSlash);
+				if (GeneratedClassPathString.StartsWith(DeveloperPathWithSlash))
+					bPassesDeveloperFilter = GeneratedClassPathString.StartsWith(UserDeveloperPathWithSlash);
 			}
 
 			// The INI files declare classes and folders that are considered internal only. Does this class match any of those patterns?
@@ -865,7 +868,7 @@ namespace ClassViewer
 			{
 				for (int i = 0; i < InternalPaths.Num(); i++)
 				{
-					if (InOriginalRootNode->GeneratedClassPackage.StartsWith(InternalPaths[i].Path))
+					if (GeneratedClassPathString.StartsWith(InternalPaths[i].Path))
 					{
 						bPassesInternalFilter = false;
 						break;
@@ -979,7 +982,7 @@ namespace ClassViewer
 				return Cast<UBlueprint>(InClass->ClassGeneratedBy);
 			}
 
-			return NULL;
+			return nullptr;
 		}
 
 		/** Retrieves a few items of information on the given UClass (retrieved via the InClassIndex). 
@@ -994,7 +997,7 @@ namespace ClassViewer
 			if (UClass* Class = InClass.Get())
 			{
 				bInOutIsBlueprintBase = CanCreateBlueprintOfClass_IgnoreDeprecation( Class );
-				bInOutHasBlueprint = Class->ClassGeneratedBy != NULL;
+				bInOutHasBlueprint = Class->ClassGeneratedBy != nullptr;
 			}
 			else
 			{
@@ -1026,27 +1029,6 @@ namespace ClassViewer
 			return false;
 		}
 
-		/** Recursively loads the entire chain of blueprints because children need their parents to be loaded.
-		 *	@param	InRootNode				The current node being examined.
-		 *	@param	InBlueprint				The blueprint to add.
-		 */
-		static void AddBlueprintChainToHierarchy(TSharedPtr< FClassViewerNode > InRootNode, UObject* InBlueprint)
-		{
-			TSharedPtr< FClassViewerNode > ParentNode = ClassHierarchy->FindParent(ClassHierarchy->GetObjectRootNode(), InRootNode->ParentClassname, NULL);
-
-			if( ParentNode.IsValid() && !ParentNode->GeneratedClassPackage.IsEmpty() )
-			{
-				UPackage* Package = LoadPackage(NULL, *ParentNode->GeneratedClassPackage, LOAD_NoRedirects );
-				Package->FullyLoad();
-				UObject* ParentObject = FindObject<UObject>(Package, *ParentNode->AssetName);
-
-				if(ParentObject->IsA(UBlueprint::StaticClass()))
-				{
-					AddBlueprintChainToHierarchy(ParentNode, ParentObject);
-				}
-			}
-		}
-
 		/**
 		 * Creates a blueprint from a class.
 		 *
@@ -1055,7 +1037,7 @@ namespace ClassViewer
 		 */
 		static void CreateBlueprint(const FString& InBlueprintName, UClass* InCreationClass)
 		{
-			if(InCreationClass == NULL || !FKismetEditorUtilities::CanCreateBlueprintOfClass(InCreationClass))
+			if(InCreationClass == nullptr || !FKismetEditorUtilities::CanCreateBlueprintOfClass(InCreationClass))
 			{
 				FMessageDialog::Open( EAppMsgType::Ok, NSLOCTEXT("UnrealEd", "InvalidClassToMakeBlueprintFrom", "Invalid class to make a Blueprint of."));
 				return;
@@ -1065,7 +1047,7 @@ namespace ClassViewer
 			FString PackageName = InBlueprintName;
 
 			// Then find/create it.
-			UPackage* Package = CreatePackage(NULL, *PackageName);
+			UPackage* Package = CreatePackage(nullptr, *PackageName);
 			check(Package);
 
 			// Handle fully loading packages before creating new objects.
@@ -1186,56 +1168,23 @@ namespace ClassViewer
 		static void LoadClass(TSharedPtr< FClassViewerNode > InOutClassNode)
 		{
 			GWarn->BeginSlowTask(LOCTEXT("LoadPackage", "Loading Package..."), true);
+			UClass* Class = LoadObject<UClass>(nullptr, *InOutClassNode->ClassPath.ToString());
+			GWarn->EndSlowTask();
 
-			UPackage* Package = LoadPackage(NULL, *InOutClassNode->GeneratedClassPackage, LOAD_NoRedirects );
-			if(Package)
+			if (Class)
 			{
-				Package->FullyLoad();
-
-				UObject* Object = FindObject<UObject>(Package, *InOutClassNode->AssetName);
-
-				GWarn->EndSlowTask();
-
-				// Check if this item is a blueprint.
-				if( Object->IsA(UBlueprint::StaticClass()) )
-				{
-					InOutClassNode->Blueprint = Cast<UBlueprint>(Object);
-					InOutClassNode->Class = *InOutClassNode->Blueprint->GeneratedClass;
-
-					// Tell the original node to update so when a refresh happens it will still know about the newly loaded class.
-					ClassViewer::Helpers::UpdateClassInNode(InOutClassNode->GeneratedClassPackage, InOutClassNode->Class.Get(), InOutClassNode->Blueprint.Get() );
-
-					// Adds the entire hierarchy of Blueprints to the EditorClassHierarchy so they will continue to appear when a full rebuild of the tree happens.
-					ClassViewer::Helpers::AddBlueprintChainToHierarchy(InOutClassNode, Object);
-				}
-				else if (UClass* Class = Cast<UClass>(Object))
-				{
-					InOutClassNode->Blueprint = Cast<UBlueprint>(Class->ClassGeneratedBy);
-					InOutClassNode->Class = Class;
-
-					// Tell the original node to update so when a refresh happens it will still know about the newly loaded class.
-					ClassViewer::Helpers::UpdateClassInNode(InOutClassNode->GeneratedClassPackage, InOutClassNode->Class.Get(), InOutClassNode->Blueprint.Get() );
+				InOutClassNode->Blueprint = Cast<UBlueprint>(Class->ClassGeneratedBy);
+				InOutClassNode->Class = Class;
 
-					// Adds the entire hierarchy of Blueprints to the EditorClassHierarchy so they will continue to appear when a full rebuild of the tree happens.
-					ClassViewer::Helpers::AddBlueprintChainToHierarchy(InOutClassNode, Object);
-				}
-				else
-				{
-					InOutClassNode->Class = Object->GetClass();
-				}
+				// Tell the original node to update so when a refresh happens it will still know about the newly loaded class.
+				ClassViewer::Helpers::UpdateClassInNode(InOutClassNode->ClassPath, InOutClassNode->Class.Get(), InOutClassNode->Blueprint.Get() );
 			}
 			else
 			{
-				GWarn->EndSlowTask();
-
-				// Check to see if the class can be found, if it can't, notify that the package failed to load.
-				//if(!FindClass(InOutClassNode))
-				{
-					FMessageLog EditorErrors("EditorErrors");
-					EditorErrors.Error(LOCTEXT("PackageLoadFail", "Package Load Failed"));
-					EditorErrors.Info(FText::FromString(InOutClassNode->GeneratedClassPackage));
-					EditorErrors.Notify(LOCTEXT("PackageLoadFail", "Package Load Failed"));
-				}
+				FMessageLog EditorErrors("EditorErrors");
+				FFormatNamedArguments Arguments;
+				Arguments.Add(TEXT("ObjectName"), FText::FromName(InOutClassNode->ClassPath));
+				EditorErrors.Error(FText::Format(LOCTEXT("PackageLoadFail", "Failed to load class {ObjectName}"), Arguments));
 			}
 		}
 
@@ -1246,7 +1195,7 @@ namespace ClassViewer
 		 */
 		static void OpenBlueprintTool(UBlueprint* InBlueprint)
 		{
-			if( InBlueprint != NULL )
+			if( InBlueprint != nullptr )
 			{
 				FAssetEditorManager::Get().OpenEditorForAsset(InBlueprint);
 			}
@@ -1287,12 +1236,12 @@ namespace ClassViewer
 		}
 
 		/** Updates the Class of a node. Uses the generated class package name to find the node.
-		*	@param InGeneratedClassPackageName			The name of the generated class package to find the node for.
-		*	@param InNewClass							The class to update the node with.
+		*	@param InGeneratedClassPath			The name of the generated class to find the node for.
+		*	@param InNewClass					The class to update the node with.
 		*/
-		static void UpdateClassInNode(const FString& InGeneratedClassPackageName, UClass* InNewClass, UBlueprint* InNewBluePrint )
+		static void UpdateClassInNode(FName InGeneratedClassPath, UClass* InNewClass, UBlueprint* InNewBluePrint )
 		{
-			ClassHierarchy->UpdateClassInNode(InGeneratedClassPackageName, InNewClass, InNewBluePrint );
+			ClassHierarchy->UpdateClassInNode(InGeneratedClassPath, InNewClass, InNewBluePrint );
 		}
 
 		static TSharedRef<SWidget> CreateMenu(UClass* Class, const bool bIsBlueprint, const bool bHasBlueprint)
@@ -1434,9 +1383,12 @@ public:
 				else if (UClass* Class = AssociatedNode->Class.Get())
 				{
 					UPackage*  Package  = Class->GetOutermost();
-					UMetaData* MetaData = Package->GetMetaData();
 					ToolTip = FEditorClassUtils::GetTooltip(Class);
 				}
+				else if (AssociatedNode->ClassPath != NAME_None)
+				{
+					ToolTip = SNew(SToolTip).Text(FText::FromName(AssociatedNode->ClassPath));
+				}
 
 				return ToolTip;
 			}
@@ -1510,7 +1462,7 @@ private:
 		// If in a Class Viewer and it has not been loaded, load the class when double-left clicking.
 		if ( bIsInClassViewer )
 		{
-			if( bDynamicClassLoading && AssociatedNode->Class == NULL && AssociatedNode->UnloadedBlueprintData.IsValid() && InMouseEvent.GetEffectingButton() == EKeys::LeftMouseButton)
+			if( bDynamicClassLoading && AssociatedNode->Class == nullptr && AssociatedNode->UnloadedBlueprintData.IsValid() && InMouseEvent.GetEffectingButton() == EKeys::LeftMouseButton)
 			{
 				ClassViewer::Helpers::LoadClass(AssociatedNode);
 			}
@@ -1664,16 +1616,10 @@ static TSharedPtr< FClassViewerNode > CreateNodeForClass(UClass* Class, const TM
 	TSharedPtr< FClassViewerNode > NewNode = MakeShareable(new FClassViewerNode(Class->GetName(), Class->GetDisplayNameText().ToString()));
 	NewNode->Blueprint = ClassViewer::Helpers::GetBlueprint(Class);
 	NewNode->Class = Class;
-
-	// Retrieve all blueprint classes
-	TArray<FAssetData> BlueprintList;
-	BlueprintPackageToAssetDataMap.MultiFind(NewNode->Class->GetOuterUPackage()->GetFName(), BlueprintList);
-
-	// Check if the Asset Registry found anything, it should, but check.
-	if ( BlueprintList.Num() )
+	NewNode->ClassPath = FName(*Class->GetPathName());
+	if (Class->GetSuperClass())
 	{
-		// Grab the generated class name and check it before assigning. Objects that haven't been saved since this has started to be exported do not have the information.
-		BlueprintList[0].GetTagValue("GeneratedClass", NewNode->GeneratedClassname);
+		NewNode->ParentClassPath = FName(*Class->GetSuperClass()->GetPathName());
 	}
 
 	return NewNode;
@@ -1717,7 +1663,7 @@ void FClassHierarchy::AddChildren_NoFilter( TSharedPtr< FClassViewerNode >& InOu
 		}
 		else
 		{
-			while ( CurrentClass->GetSuperClass() != NULL )
+			while ( CurrentClass->GetSuperClass() != nullptr )
 			{
 				TSharedPtr<FClassViewerNode>& ParentEntry = Nodes.FindOrAdd(CurrentClass->GetSuperClass());
 				if ( !ParentEntry.IsValid() )
@@ -1746,7 +1692,7 @@ void FClassHierarchy::AddChildren_NoFilter( TSharedPtr< FClassViewerNode >& InOu
 TSharedPtr< FClassViewerNode > FClassHierarchy::FindParent(const TSharedPtr< FClassViewerNode >& InRootNode, FName InParentClassname, const UClass* InParentClass)
 {
 	// Check if the current node is the parent classname that is being searched for.
-	if(InRootNode->GeneratedClassname == InParentClassname)
+	if(InRootNode->ClassPath == InParentClassname)
 	{
 		// Return the node if it is the correct parent, this ends the recursion.
 		return InRootNode;
@@ -1808,9 +1754,9 @@ TSharedPtr< FClassViewerNode > FClassHierarchy::FindNodeByClassName(const TShare
 	return ReturnNode;
 }
 
-TSharedPtr< FClassViewerNode > FClassHierarchy::FindNodeByGeneratedClassPackageName(const TSharedPtr< FClassViewerNode >& InRootNode, const FString& InGeneratedClassPackageName)
+TSharedPtr< FClassViewerNode > FClassHierarchy::FindNodeByGeneratedClassPath(const TSharedPtr< FClassViewerNode >& InRootNode, FName InGeneratedClassPath)
 {
-	if(InRootNode->GeneratedClassPackage == InGeneratedClassPackageName)
+	if(InRootNode->ClassPath == InGeneratedClassPath)
 	{
 		return InRootNode;
 	}
@@ -1821,7 +1767,7 @@ TSharedPtr< FClassViewerNode > FClassHierarchy::FindNodeByGeneratedClassPackageN
 	for(int32 ChildClassIndex = 0; !ReturnNode.IsValid() && ChildClassIndex < InRootNode->GetChildrenList().Num(); ChildClassIndex++)
 	{
 		// Check the child, then check the return to see if it is valid. If it is valid, end the recursion.
-		ReturnNode = FindNodeByGeneratedClassPackageName(InRootNode->GetChildrenList()[ChildClassIndex], InGeneratedClassPackageName);
+		ReturnNode = FindNodeByGeneratedClassPath(InRootNode->GetChildrenList()[ChildClassIndex], InGeneratedClassPath);
 
 		if(ReturnNode.IsValid())
 		{
@@ -1832,9 +1778,9 @@ TSharedPtr< FClassViewerNode > FClassHierarchy::FindNodeByGeneratedClassPackageN
 	return ReturnNode;
 }
 
-void FClassHierarchy::UpdateClassInNode(const FString& InGeneratedClassPackageName, UClass* InNewClass,  UBlueprint* InNewBluePrint )
+void FClassHierarchy::UpdateClassInNode(FName InGeneratedClassPath, UClass* InNewClass, UBlueprint* InNewBluePrint )
 {
-	TSharedPtr< FClassViewerNode > Node = FindNodeByGeneratedClassPackageName(ObjectClassRoot, InGeneratedClassPackageName);
+	TSharedPtr< FClassViewerNode > Node = FindNodeByGeneratedClassPath(ObjectClassRoot, InGeneratedClassPath);
 
 	if( Node.IsValid() )
 	{
@@ -1843,20 +1789,20 @@ void FClassHierarchy::UpdateClassInNode(const FString& InGeneratedClassPackageNa
 	}
 }
 
-bool FClassHierarchy::FindAndRemoveNodeByPackageName(const TSharedPtr< FClassViewerNode >& InRootNode, const FString& InAssetPackage)
+bool FClassHierarchy::FindAndRemoveNodeByClassPath(const TSharedPtr< FClassViewerNode >& InRootNode, FName InClassPath)
 {
 	bool bReturnValue = false;
 
 	// Search the children recursively, one of them might have the parent.
 	for(int32 ChildClassIndex = 0; ChildClassIndex < InRootNode->GetChildrenList().Num(); ChildClassIndex++)
 	{
-		if(InRootNode->GetChildrenList()[ChildClassIndex]->GeneratedClassPackage == InAssetPackage)
+		if(InRootNode->GetChildrenList()[ChildClassIndex]->ClassPath == InClassPath)
 		{
 			InRootNode->GetChildrenList().RemoveAt(ChildClassIndex);
 			return true;
 		}
 		// Check the child, then check the return to see if it is valid. If it is valid, end the recursion.
-		bReturnValue = FindAndRemoveNodeByPackageName(InRootNode->GetChildrenList()[ChildClassIndex], InAssetPackage);
+		bReturnValue = FindAndRemoveNodeByClassPath(InRootNode->GetChildrenList()[ChildClassIndex], InClassPath);
 
 		if(bReturnValue)
 		{
@@ -1868,7 +1814,13 @@ bool FClassHierarchy::FindAndRemoveNodeByPackageName(const TSharedPtr< FClassVie
 
 void FClassHierarchy::RemoveAsset(const FAssetData& InRemovedAssetData)
 {
-	if(FindAndRemoveNodeByPackageName(ObjectClassRoot, InRemovedAssetData.PackageName.ToString()))
+	FString ClassObjectPath;
+	if (InRemovedAssetData.GetTagValue(FBlueprintTags::GeneratedClassPath, ClassObjectPath))
+	{
+		ClassObjectPath = FPackageName::ExportTextPathToObjectPath(ClassObjectPath);
+	}
+
+	if (FindAndRemoveNodeByClassPath(ObjectClassRoot, FName(*ClassObjectPath)))
 	{
 		// All viewers must refresh.
 		ClassViewer::Helpers::RefreshAll();
@@ -1885,8 +1837,14 @@ void FClassHierarchy::AddAsset(const FAssetData& InAddedAssetData)
 
 		if( AncestorClassNames.Contains(UBlueprintCore::StaticClass()->GetFName()) )
 		{
+			FString ClassObjectPath;
+			if (InAddedAssetData.GetTagValue(FBlueprintTags::GeneratedClassPath, ClassObjectPath))
+			{
+				ClassObjectPath = FPackageName::ExportTextPathToObjectPath(ClassObjectPath);
+			}
+
 			// Make sure that the node does not already exist. There is a bit of double adding going on at times and this prevents it.
-			if(!FindNodeByGeneratedClassPackageName(ObjectClassRoot, InAddedAssetData.PackageName.ToString()).IsValid())
+			if(!FindNodeByGeneratedClassPath(ObjectClassRoot, FName(*ClassObjectPath)).IsValid())
 			{
 				TSharedPtr< FClassViewerNode > NewNode;
 				LoadUnloadedTagData(NewNode, InAddedAssetData);
@@ -1895,12 +1853,9 @@ void FClassHierarchy::AddAsset(const FAssetData& InAddedAssetData)
 				FindClass(NewNode);
 
 				// Resolve the parent's class name locally and use it to find the parent's class.
-				FString ParentClassName = NewNode->ParentClassname.ToString();
-				UObject* Outer(NULL);
-				ResolveName(Outer, ParentClassName, false, false);
-
-				UClass* ParentClass = FindObject<UClass>(ANY_PACKAGE, *ParentClassName);
-				TSharedPtr< FClassViewerNode > ParentNode = FindParent(ObjectClassRoot, NewNode->ParentClassname, ParentClass); 
+				FString ParentClassPath = NewNode->ParentClassPath.ToString();
+				UClass* ParentClass = FindObject<UClass>(nullptr, *ParentClassPath);
+				TSharedPtr< FClassViewerNode > ParentNode = FindParent(ObjectClassRoot, NewNode->ParentClassPath, ParentClass); 
 				if(ParentNode.IsValid())
 				{
 					ParentNode->AddChild(NewNode);
@@ -1934,56 +1889,46 @@ void FClassHierarchy::SortChildren( TSharedPtr< FClassViewerNode >& InRootNode )
 
 void FClassHierarchy::FindClass(TSharedPtr< FClassViewerNode > InOutClassNode)
 {
-	UPackage* Package = FindPackage(NULL, *InOutClassNode->GeneratedClassPackage );
-	if (Package)
-	{
-		UObject* Object = FindObject<UObject>(Package, *InOutClassNode->AssetName);
+	UClass* Class = FindObject<UClass>(nullptr, *InOutClassNode->ClassPath.ToString());
 
-		if (Object)
-		{
-			// Check if this item is a blueprint.
-			if (Object->IsA(UBlueprint::StaticClass()))
-			{
-				InOutClassNode->Blueprint = Cast<UBlueprint>(Object);
-				if (InOutClassNode->Blueprint.IsValid())
-				{
-					InOutClassNode->Class = *InOutClassNode->Blueprint->GeneratedClass;
-				}
-			}
-			else if (UClass* Class = Cast<UClass>(Object))
-			{
-				InOutClassNode->Blueprint = Cast<UBlueprint>(Class->ClassGeneratedBy);
-				InOutClassNode->Class = Class;
-			}
-			else
-			{
-				InOutClassNode->Class = Object->GetClass();
-			}
-		}
+	if (Class)
+	{
+		InOutClassNode->Blueprint = Cast<UBlueprint>(Class->ClassGeneratedBy);
+		InOutClassNode->Class = Class;
 	}
 }
 
 void FClassHierarchy::LoadUnloadedTagData(TSharedPtr<FClassViewerNode>& InOutClassViewerNode, const FAssetData& InAssetData)
 {
-	// Create the viewer node.
-	const FString AssetName = InAssetData.AssetName.ToString();
-	InOutClassViewerNode = MakeShareable(new FClassViewerNode(AssetName, AssetName));
+	// Create the viewer node. We use the name without _C for both
+	const FString ClassName = InAssetData.AssetName.ToString();
+	InOutClassViewerNode = MakeShareable(new FClassViewerNode(ClassName, ClassName));
+
+	InOutClassViewerNode->BlueprintAssetPath = InAssetData.ObjectPath;
+
+	FString ClassObjectPath;
+	if (InAssetData.GetTagValue(FBlueprintTags::GeneratedClassPath, ClassObjectPath))
+	{
+		InOutClassViewerNode->ClassPath = FName(*FPackageName::ExportTextPathToObjectPath(ClassObjectPath));
+	}
 
-	InOutClassViewerNode->AssetName = AssetName;
-	InOutClassViewerNode->GeneratedClassPackage = InAssetData.PackageName.ToString();
-	InAssetData.GetTagValue("GeneratedClass", InOutClassViewerNode->GeneratedClassname);
-	InAssetData.GetTagValue("ParentClass", InOutClassViewerNode->ParentClassname);
-	InOutClassViewerNode->bIsBPNormalType = InAssetData.GetTagValueRef<FString>("BlueprintType") == TEXT("BPType_Normal");
+	FString ParentClassPathString;
+	if (InAssetData.GetTagValue(FBlueprintTags::ParentClassPath, ParentClassPathString))
+	{
+		InOutClassViewerNode->ParentClassPath = FName(*FPackageName::ExportTextPathToObjectPath(ParentClassPathString));
+	}
+
+	InOutClassViewerNode->bIsBPNormalType = InAssetData.GetTagValueRef<FString>(FBlueprintTags::BlueprintType) == TEXT("BPType_Normal");
 			
 	// It is an unloaded blueprint, so we need to create the structure that will hold the data.
 	TSharedPtr<FUnloadedBlueprintData> UnloadedBlueprintData = MakeShareable( new FUnloadedBlueprintData(InOutClassViewerNode) );
 	InOutClassViewerNode->UnloadedBlueprintData = UnloadedBlueprintData;
 
 	// Get the class flags.
-	const uint32 ClassFlags = InAssetData.GetTagValueRef<uint32>("ClassFlags");
+	const uint32 ClassFlags = InAssetData.GetTagValueRef<uint32>(FBlueprintTags::ClassFlags);
 	InOutClassViewerNode->UnloadedBlueprintData->SetClassFlags(ClassFlags);
 
-	const FString ImplementedInterfaces = InAssetData.GetTagValueRef<FString>("ImplementedInterfaces");
+	const FString ImplementedInterfaces = InAssetData.GetTagValueRef<FString>(FBlueprintTags::ImplementedInterfaces);
 	if(!ImplementedInterfaces.IsEmpty())
 	{
 		FString FullInterface;
@@ -2046,17 +1991,15 @@ void FClassHierarchy::PopulateClassHierarchy()
 	// Second pass to link them to parents.
 	for (int32 CurrentNodeIdx = 0; CurrentNodeIdx < RootLevelClasses.Num(); ++CurrentNodeIdx)
 	{
-		if(RootLevelClasses[CurrentNodeIdx]->ParentClassname != NAME_None)
+		if(RootLevelClasses[CurrentNodeIdx]->ParentClassPath != NAME_None)
 		{
 			// Resolve the parent's class name locally and use it to find the parent's class.
-			FString ParentClassName = RootLevelClasses[CurrentNodeIdx]->ParentClassname.ToString();
-			UObject* Outer(NULL);
-			ResolveName(Outer, ParentClassName, false, false);
-			const UClass* ParentClass = FindObject<UClass>(ANY_PACKAGE, *ParentClassName);
+			FString ParentClassPath = RootLevelClasses[CurrentNodeIdx]->ParentClassPath.ToString();
+			const UClass* ParentClass = FindObject<UClass>(nullptr, *ParentClassPath);
 
 			for (int32 SearchNodeIdx = 0; SearchNodeIdx < RootLevelClasses.Num(); ++SearchNodeIdx)
 			{
-				TSharedPtr< FClassViewerNode > ParentNode = FindParent(RootLevelClasses[SearchNodeIdx], RootLevelClasses[CurrentNodeIdx]->ParentClassname, ParentClass); 
+				TSharedPtr< FClassViewerNode > ParentNode = FindParent(RootLevelClasses[SearchNodeIdx], RootLevelClasses[CurrentNodeIdx]->ParentClassPath, ParentClass);
 				if(ParentNode.IsValid())
 				{
 					// AddUniqueChild makes sure that when a node was generated one by EditorClassHierarchy and one from LoadUnloadedTagData - the proper one is selected
@@ -2439,7 +2382,7 @@ void SClassViewer::OnClassViewerSelectionChanged( TSharedPtr<FClassViewerNode> I
 	{
 		UClass* Class = Item->Class.Get();
 
-		// If the class is NULL and UnloadedBlueprintData is valid then attempt to load it. UnloadedBlueprintData is invalid in the case of a "None" item.
+		// If the class is nullptr and UnloadedBlueprintData is valid then attempt to load it. UnloadedBlueprintData is invalid in the case of a "None" item.
 		if ( bEnableClassDynamicLoading && !Class && Item->UnloadedBlueprintData.IsValid() )
 		{
 			ClassViewer::Helpers::LoadClass( Item );
@@ -2454,7 +2397,7 @@ void SClassViewer::OnClassViewerSelectionChanged( TSharedPtr<FClassViewerNode> I
 			}
 			else
 			{
-				OnClassPicked.ExecuteIfBound( NULL );
+				OnClassPicked.ExecuteIfBound( nullptr );
 			}
 		}
 	}
@@ -2530,11 +2473,6 @@ TSharedRef< ITableRow > SClassViewer::OnGenerateRowForClassViewer( TSharedPtr<FC
 		.OnDragDetected(this, &SClassViewer::OnDragDetected)
 		.OnClassItemDoubleClicked(FOnClassItemDoubleClickDelegate::CreateSP(this, &SClassViewer::ToggleExpansionState_Helper));
 
-	if ( !Item->GeneratedClassPackage.IsEmpty() )
-	{
-		ReturnRow->SetToolTipText(FText::FromString(Item->GeneratedClassPackage));
-	}
-
 	// Expand the item if needed.
 	if (!bPendingSetExpansionStates)
 	{
@@ -2576,7 +2514,7 @@ TSharedRef<SWidget> SClassViewer::GetViewButtonContent()
 {
 	// Get all menu extenders for this context menu from the content browser module
 
-	FMenuBuilder MenuBuilder(/*bInShouldCloseWindowAfterMenuSelection=*/true, NULL, NULL, /*bCloseSelfOnly=*/ true);
+	FMenuBuilder MenuBuilder(/*bInShouldCloseWindowAfterMenuSelection=*/true, nullptr, nullptr, /*bCloseSelfOnly=*/ true);
 
 	MenuBuilder.AddMenuEntry(LOCTEXT("ExpandAll", "Expand All"), LOCTEXT("ExpandAll_Tooltip", "Expands the entire tree"), FSlateIcon(), FUIAction(FExecuteAction::CreateSP(this, &SClassViewer::SetAllExpansionStates, bool(true))), NAME_None, EUserInterfaceActionType::Button);
 	MenuBuilder.AddMenuEntry(LOCTEXT("CollapseAll", "Collapse All"), LOCTEXT("CollapseAll_Tooltip", "Collapses the entire tree"), FSlateIcon(), FUIAction(FExecuteAction::CreateSP(this, &SClassViewer::SetAllExpansionStates, bool(false))), NAME_None, EUserInterfaceActionType::Button);
@@ -2745,7 +2683,7 @@ FReply SClassViewer::OnDragDetected( const FGeometry& Geometry, const FPointerEv
 		{
 			TSharedRef< FClassViewerNode > Item = SelectedItems[0].ToSharedRef();
 
-			// If there is no class then we must spawn an FUnloadedClassDragDropOp so the class will be loaded when dropped.
+			// If there is no class then we must spawn an FAssetDragDropOp so the class will be loaded when dropped.
 			if ( UClass* Class = Item->Class.Get() )
 			{
 				// Spawn a loaded blueprint just like any other asset from the Content Browser.
@@ -2761,9 +2699,14 @@ FReply SClassViewer::OnDragDetected( const FGeometry& Geometry, const FPointerEv
 					return FReply::Handled().BeginDragDrop(FClassDragDropOp::New(MakeWeakObjectPtr(Class)));
 				}	
 			}
-			else
+			else if (Item->BlueprintAssetPath != NAME_None)
 			{
-				return FReply::Handled().BeginDragDrop(FUnloadedClassDragDropOp::New(FClassPackageData(Item->AssetName, Item->GeneratedClassPackage)));
+				FAssetRegistryModule& AssetRegistryModule = FModuleManager::LoadModuleChecked<FAssetRegistryModule>(TEXT("AssetRegistry"));
+
+				// Pull asset data out of asset registry
+				TArray<FAssetData> InAssetData;
+				InAssetData.Add(AssetRegistryModule.Get().GetAssetByObjectPath(Item->BlueprintAssetPath));
+				return FReply::Handled().BeginDragDrop(FAssetDragDropOp::New(InAssetData));
 			}
 		}
 	}
@@ -2800,14 +2743,14 @@ void SClassViewer::OnFilterTextCommitted(const FText& InText, ETextCommit::Type
 			TArray< TSharedPtr< FClassViewerNode > > SelectedList = ClassList->GetSelectedItems();
 			TSharedPtr< FClassViewerNode > FirstSelected;
 
-			UClass* Class = NULL;
+			UClass* Class = nullptr;
 			
 			if (SelectedList.Num() > 0)
 			{
 				FirstSelected = SelectedList[0];
 				Class = FirstSelected->Class.Get();
 
-				// If the class is NULL and UnloadedBlueprintData is valid then attempt to load it. UnloadedBlueprintData is invalid in the case of a "None" item.
+				// If the class is nullptr and UnloadedBlueprintData is valid then attempt to load it. UnloadedBlueprintData is invalid in the case of a "None" item.
 				if ( bEnableClassDynamicLoading && Class == nullptr && FirstSelected->UnloadedBlueprintData.IsValid())
 				{
 					ClassViewer::Helpers::LoadClass(FirstSelected);
