SClassViewer was not displaying the full set of classes when creating a new blueprint due to bReturnPassesFilter defaulting to false.

diff --git a/Engine/Source/Editor/ClassViewer/Private/SClassViewer.cpp b/Engine/Source/Editor/ClassViewer/Private/SClassViewer.cpp
--- a/Engine/Source/Editor/ClassViewer/Private/SClassViewer.cpp
+++ b/Engine/Source/Editor/ClassViewer/Private/SClassViewer.cpp
@@ -315,7 +315,7 @@ namespace ClassViewer
 				return false;
 			}
 
-			bool bReturnPassesFilter = false;
+			bool bReturnPassesFilter = true;
 
 			TArray< TSharedPtr< FClassViewerNode > >& ChildList = InOriginalRootNode->GetChildrenList();
 			for(int32 ChildIdx = 0; ChildIdx < ChildList.Num(); ChildIdx++)
