Fixed filtering of classes by a search string when creating a new blueprint.

diff --git a/Engine/Source/Editor/ClassViewer/Private/SClassViewer.cpp b/Engine/Source/Editor/ClassViewer/Private/SClassViewer.cpp
--- a/Engine/Source/Editor/ClassViewer/Private/SClassViewer.cpp
+++ b/Engine/Source/Editor/ClassViewer/Private/SClassViewer.cpp
@@ -310,12 +310,8 @@ namespace ClassViewer
 			const TSharedPtr<FClassViewerFilter>& InClassFilter, const FClassViewerInitializationOptions& InInitOptions)
 		{
 			InOutRootNode->bPassesFilter = InClassFilter->IsNodeAllowed(InInitOptions, InOutRootNode.ToSharedRef());
-			if (!InOutRootNode->bPassesFilter)
-			{
-				return false;
-			}
 
-			bool bReturnPassesFilter = true;
+			bool bReturnPassesFilter = InOutRootNode->bPassesFilter;
 
 			TArray< TSharedPtr< FClassViewerNode > >& ChildList = InOriginalRootNode->GetChildrenList();
 			for(int32 ChildIdx = 0; ChildIdx < ChildList.Num(); ChildIdx++)
