#UE4 Nonunity CIS fix

diff --git a/Engine/Source/Editor/ClassViewer/Public/ClassViewerModule.h b/Engine/Source/Editor/ClassViewer/Public/ClassViewerModule.h
--- a/Engine/Source/Editor/ClassViewer/Public/ClassViewerModule.h
+++ b/Engine/Source/Editor/ClassViewer/Public/ClassViewerModule.h
@@ -3,6 +3,7 @@
 #pragma once
 
 #include "CoreMinimal.h"
+#include "AssetData.h"
 #include "Modules/ModuleInterface.h"
 
 class IClassViewerFilter;
