#UE4 ClassViewerFilter integration with AssetReferenceFilter that was lost in merging to Dev-EngineMerge

diff --git a/Engine/Source/Editor/ClassViewer/Private/ClassViewerFilter.cpp b/Engine/Source/Editor/ClassViewer/Private/ClassViewerFilter.cpp
--- a/Engine/Source/Editor/ClassViewer/Private/ClassViewerFilter.cpp
+++ b/Engine/Source/Editor/ClassViewer/Private/ClassViewerFilter.cpp
@@ -9,6 +9,10 @@
 #include "Misc/PackageName.h"
 #include "Misc/Paths.h"
 #include "Misc/TextFilterExpressionEvaluator.h"
+#include "Editor/UnrealEdEngine.h"
+#include "UnrealEdGlobals.h"
+#include "AssetRegistryModule.h"
+#include "PropertyHandle.h"
 
 EFilterReturn::Type FClassViewerFilterFuncs::IfInChildOfClassesSet(TSet< const UClass* >& InSet, const UClass* InClass)
 {
@@ -390,8 +394,25 @@ static bool PassesTextFilter(const FString& InTestString, const TSharedRef<FText
 
 FClassViewerFilter::FClassViewerFilter(const FClassViewerInitializationOptions& InInitOptions) :
 	TextFilter(MakeShared<FTextFilterExpressionEvaluator>(ETextFilterExpressionEvaluatorMode::BasicString)),
-	FilterFunctions(MakeShared<FClassViewerFilterFuncs>())
+	FilterFunctions(MakeShared<FClassViewerFilterFuncs>()),
+	AssetRegistry(FModuleManager::GetModuleChecked<FAssetRegistryModule>(TEXT("AssetRegistry")).Get())
 {
+	// Create a game-specific filter, if the referencing property/assets were supplied
+	if (GUnrealEd)
+	{
+		FAssetReferenceFilterContext AssetReferenceFilterContext;
+		AssetReferenceFilterContext.ReferencingAssets = InInitOptions.AdditionalReferencingAssets;
+		if (InInitOptions.PropertyHandle.IsValid())
+		{
+			TArray<UObject*> ReferencingObjects;
+			InInitOptions.PropertyHandle->GetOuterObjects(ReferencingObjects);
+			for (UObject* ReferencingObject : ReferencingObjects)
+			{
+				AssetReferenceFilterContext.ReferencingAssets.Add(FAssetData(ReferencingObject));
+			}
+		}
+		AssetReferenceFilter = GUnrealEd->MakeAssetReferenceFilter(AssetReferenceFilterContext);
+	}
 }
 
 bool FClassViewerFilter::IsNodeAllowed(const FClassViewerInitializationOptions& InInitOptions, const TSharedRef<FClassViewerNode>& InNode)
@@ -479,9 +500,15 @@ bool FClassViewerFilter::IsClassAllowed(const FClassViewerInitializationOptions&
 
 	const bool bPassesTextFilter = PassesTextFilter(InClass->GetName(), TextFilter);
 
+	bool bPassesAssetReferenceFilter = true;
+	if (AssetReferenceFilter.IsValid() && !InClass->IsNative())
+	{
+		bPassesAssetReferenceFilter = AssetReferenceFilter->PassesFilter(FAssetData(InClass));
+	}
+
 	bool bPassesFilter = bPassesPlaceableFilter && bPassesBlueprintBaseFilter 
 		&& bPassesDeveloperFilter && bPassesInternalFilter && bPassesEditorClassFilter 
-		&& bPassesCustomFilter && bPassesTextFilter;
+		&& bPassesCustomFilter && bPassesTextFilter && bPassesAssetReferenceFilter;
 
 	return bPassesFilter;
 }
@@ -493,7 +520,8 @@ bool FClassViewerFilter::IsUnloadedClassAllowed(const FClassViewerInitialization
 		return false;
 	}
 
-	const bool bPassesBlueprintBaseFilter = !InInitOptions.bIsBlueprintBaseOnly || CheckIfBlueprintBase(InUnloadedClassData);
+	const bool bIsBlueprintBase = CheckIfBlueprintBase(InUnloadedClassData);
+	const bool bPassesBlueprintBaseFilter = !InInitOptions.bIsBlueprintBaseOnly || bIsBlueprintBase;
 
 	// unloaded blueprints cannot be editor-only
 	const bool bPassesEditorClassFilter = !InInitOptions.bEditorClassesOnly;
@@ -548,9 +576,16 @@ bool FClassViewerFilter::IsUnloadedClassAllowed(const FClassViewerInitialization
 
 	const bool bPassesTextFilter = PassesTextFilter(*InUnloadedClassData->GetClassName().Get(), TextFilter);
 
+	bool bPassesAssetReferenceFilter = true;
+	if (AssetReferenceFilter.IsValid() && bIsBlueprintBase)
+	{
+		FName BlueprintPath = FName(*GeneratedClassPathString.LeftChop(2)); // Chop off _C
+		bPassesAssetReferenceFilter = AssetReferenceFilter->PassesFilter(AssetRegistry.GetAssetByObjectPath(BlueprintPath));
+	}
+
 	bool bPassesFilter = bPassesPlaceableFilter && bPassesBlueprintBaseFilter
 		&& bPassesDeveloperFilter && bPassesInternalFilter && bPassesEditorClassFilter
-		&& bPassesCustomFilter && bPassesTextFilter;
+		&& bPassesCustomFilter && bPassesTextFilter && bPassesAssetReferenceFilter;
 
 	return bPassesFilter;
 }
diff --git a/Engine/Source/Editor/ClassViewer/Public/ClassViewerFilter.h b/Engine/Source/Editor/ClassViewer/Public/ClassViewerFilter.h
--- a/Engine/Source/Editor/ClassViewer/Public/ClassViewerFilter.h
+++ b/Engine/Source/Editor/ClassViewer/Public/ClassViewerFilter.h
@@ -8,6 +8,8 @@
 class FClassViewerNode;
 class FClassViewerInitializationOptions;
 class FTextFilterExpressionEvaluator;
+class IAssetReferenceFilter;
+class IAssetRegistry;
 
 /** Interface class for creating filters for the Class Viewer. */
 class IClassViewerFilter
@@ -50,6 +52,8 @@ public:
 
 	TSharedRef<FTextFilterExpressionEvaluator> TextFilter;
 	TSharedRef<FClassViewerFilterFuncs> FilterFunctions;
+	TSharedPtr<IAssetReferenceFilter> AssetReferenceFilter;
+	const IAssetRegistry& AssetRegistry;
 };
 
 namespace EFilterReturn
