Fix for window flickering

diff --git a/Engine/Source/Editor/ClassViewer/Private/SClassViewer.cpp b/Engine/Source/Editor/ClassViewer/Private/SClassViewer.cpp
--- a/Engine/Source/Editor/ClassViewer/Private/SClassViewer.cpp
+++ b/Engine/Source/Editor/ClassViewer/Private/SClassViewer.cpp
@@ -2162,134 +2162,138 @@ void SClassViewer::Construct(const FArguments& InArgs, const FClassViewerInitial
 	TSharedPtr< SWidget > ClassViewerContent;
 
 	ClassViewerContent = 
-	SNew(SBorder)
-	.BorderImage(FEditorStyle::GetBrush("ToolPanel.GroupBorder"))
+	SNew(SBox)
+	.MaxDesiredHeight(800.0f)
 	[
-		SNew(SVerticalBox)
-		+SVerticalBox::Slot()
-		.AutoHeight()
-		.Padding( 1.0f, 0.0f, 1.0f, 0.0f )
+		SNew(SBorder)
+		.BorderImage(FEditorStyle::GetBrush("ToolPanel.GroupBorder"))
 		[
-			SNew(SHorizontalBox)
-			+SHorizontalBox::Slot()
-			.AutoWidth()
-			.VAlign(VAlign_Center)
+			SNew(SVerticalBox)
+			+SVerticalBox::Slot()
+			.AutoHeight()
+			.Padding( 1.0f, 0.0f, 1.0f, 0.0f )
 			[
-				SNew(STextBlock)
-				.Visibility(bHasTitle ? EVisibility::Visible : EVisibility::Collapsed)
-				.ColorAndOpacity(FEditorStyle::GetColor("MultiboxHookColor"))
-				.Text(InitOptions.ViewerTitleString)
+				SNew(SHorizontalBox)
+				+SHorizontalBox::Slot()
+				.AutoWidth()
+				.VAlign(VAlign_Center)
+				[
+					SNew(STextBlock)
+					.Visibility(bHasTitle ? EVisibility::Visible : EVisibility::Collapsed)
+					.ColorAndOpacity(FEditorStyle::GetColor("MultiboxHookColor"))
+					.Text(InitOptions.ViewerTitleString)
+				]
 			]
-		]
-		+SVerticalBox::Slot()
-		.AutoHeight()
-		[
-			SNew(SHorizontalBox)
-			+ SHorizontalBox::Slot()
-			.AutoWidth()
-			.Padding(2.0f, 2.0f)
+			+SVerticalBox::Slot()
+			.AutoHeight()
 			[
-				FiltersWidget
+				SNew(SHorizontalBox)
+				+ SHorizontalBox::Slot()
+				.AutoWidth()
+				.Padding(2.0f, 2.0f)
+				[
+					FiltersWidget
+				]
+				+ SHorizontalBox::Slot()
+				.Padding(2.0f, 2.0f)
+				[
+					SAssignNew(SearchBox, SSearchBox)
+					.OnTextChanged( this, &SClassViewer::OnFilterTextChanged )
+					.OnTextCommitted( this, &SClassViewer::OnFilterTextCommitted )
+				]
 			]
-			+ SHorizontalBox::Slot()
-			.Padding(2.0f, 2.0f)
+
+			+SVerticalBox::Slot()
+			.AutoHeight()
 			[
-				SAssignNew(SearchBox, SSearchBox)
-				.OnTextChanged( this, &SClassViewer::OnFilterTextChanged )
-				.OnTextCommitted( this, &SClassViewer::OnFilterTextCommitted )
+				SNew(SSeparator)
+				.Visibility(HeaderVisibility)
 			]
-		]
-
-		+SVerticalBox::Slot()
-		.AutoHeight()
-		[
-			SNew(SSeparator)
-			.Visibility(HeaderVisibility)
-		]
-
-		+SVerticalBox::Slot()
-		.FillHeight(1.0f)
-		[
-			SNew(SOverlay)
 
-			+SOverlay::Slot()
-			.HAlign(HAlign_Fill)
-			.VAlign(VAlign_Fill)
+			+SVerticalBox::Slot()
+			.FillHeight(1.0f)
 			[
-				SNew(SVerticalBox)
+				SNew(SOverlay)
 
-				+SVerticalBox::Slot()
-				.FillHeight(1.0f)
+				+SOverlay::Slot()
+				.HAlign(HAlign_Fill)
+				.VAlign(VAlign_Fill)
 				[
-					SNew(SScrollBorder, ClassTreeView)
-					.Visibility(InitOptions.DisplayMode == EClassViewerDisplayMode::TreeView ? EVisibility::Visible : EVisibility::Collapsed)
+					SNew(SVerticalBox)
+
+					+SVerticalBox::Slot()
+					.FillHeight(1.0f)
 					[
-						ClassTreeView
+						SNew(SScrollBorder, ClassTreeView)
+						.Visibility(InitOptions.DisplayMode == EClassViewerDisplayMode::TreeView ? EVisibility::Visible : EVisibility::Collapsed)
+						[
+							ClassTreeView
+						]
 					]
-				]
 
-				+SVerticalBox::Slot()
-				.FillHeight(1.0f)
-				[
-					SNew(SScrollBorder, ClassListView)
-					.Visibility(InitOptions.DisplayMode == EClassViewerDisplayMode::ListView ? EVisibility::Visible : EVisibility::Collapsed)
+					+SVerticalBox::Slot()
+					.FillHeight(1.0f)
 					[
-						ClassListView
+						SNew(SScrollBorder, ClassListView)
+						.Visibility(InitOptions.DisplayMode == EClassViewerDisplayMode::ListView ? EVisibility::Visible : EVisibility::Collapsed)
+						[
+							ClassListView
+						]
 					]
 				]
-			]
 
-			+SOverlay::Slot()
-			.HAlign(HAlign_Fill)
-			.VAlign(VAlign_Bottom)
-			.Padding(FMargin(24, 0, 24, 0))
-			[
-				// Asset discovery indicator
-				AssetDiscoveryIndicator
+				+SOverlay::Slot()
+				.HAlign(HAlign_Fill)
+				.VAlign(VAlign_Bottom)
+				.Padding(FMargin(24, 0, 24, 0))
+				[
+					// Asset discovery indicator
+					AssetDiscoveryIndicator
+				]
 			]
-		]
-
-		// Bottom panel
-		+ SVerticalBox::Slot()
-		.AutoHeight()
-		[
-			SNew(SHorizontalBox)
 
-			// Asset count
-			+ SHorizontalBox::Slot()
-			.FillWidth(1.f)
-			.VAlign(VAlign_Center)
-			.Padding(8, 0)
+			// Bottom panel
+			+ SVerticalBox::Slot()
+			.AutoHeight()
 			[
-				SNew(STextBlock)
-				.Text(this, &SClassViewer::GetClassCountText)
-			]
+				SNew(SHorizontalBox)
 
-			// View mode combo button
-			+ SHorizontalBox::Slot()
-			.AutoWidth()
-			[
-				SAssignNew(ViewOptionsComboButton, SComboButton)
-				.ContentPadding(0)
-				.ForegroundColor(this, &SClassViewer::GetViewButtonForegroundColor)
-				.ButtonStyle(FEditorStyle::Get(), "ToggleButton") // Use the tool bar item style for this button
-				.OnGetMenuContent(this, &SClassViewer::GetViewButtonContent)
-				.ButtonContent()
+				// Asset count
+				+ SHorizontalBox::Slot()
+				.FillWidth(1.f)
+				.VAlign(VAlign_Center)
+				.Padding(8, 0)
 				[
-					SNew(SHorizontalBox)
-						+ SHorizontalBox::Slot()
-					.AutoWidth()
-					.VAlign(VAlign_Center)
-					[
-						SNew(SImage).Image(FEditorStyle::GetBrush("GenericViewButton"))
-					]
+					SNew(STextBlock)
+					.Text(this, &SClassViewer::GetClassCountText)
+				]
 
-					+ SHorizontalBox::Slot()
-					.AutoWidth()
-					.Padding(2, 0, 0, 0)
-					.VAlign(VAlign_Center)
+				// View mode combo button
+				+ SHorizontalBox::Slot()
+				.AutoWidth()
+				[
+					SAssignNew(ViewOptionsComboButton, SComboButton)
+					.ContentPadding(0)
+					.ForegroundColor(this, &SClassViewer::GetViewButtonForegroundColor)
+					.ButtonStyle(FEditorStyle::Get(), "ToggleButton") // Use the tool bar item style for this button
+					.OnGetMenuContent(this, &SClassViewer::GetViewButtonContent)
+					.ButtonContent()
 					[
-						SNew(STextBlock).Text(LOCTEXT("ViewButton", "View Options"))
+						SNew(SHorizontalBox)
+							+ SHorizontalBox::Slot()
+						.AutoWidth()
+						.VAlign(VAlign_Center)
+						[
+							SNew(SImage).Image(FEditorStyle::GetBrush("GenericViewButton"))
+						]
+
+						+ SHorizontalBox::Slot()
+						.AutoWidth()
+						.Padding(2, 0, 0, 0)
+						.VAlign(VAlign_Center)
+						[
+							SNew(STextBlock).Text(LOCTEXT("ViewButton", "View Options"))
+						]
 					]
 				]
 			]
