PR #5108: Make SClassViewer options more readable in GameMode BP (Contributed by [REDACTED_PERSON])

diff --git a/Engine/Source/Editor/ClassViewer/Private/SClassViewer.cpp b/Engine/Source/Editor/ClassViewer/Private/SClassViewer.cpp
--- a/Engine/Source/Editor/ClassViewer/Private/SClassViewer.cpp
+++ b/Engine/Source/Editor/ClassViewer/Private/SClassViewer.cpp
@@ -2163,12 +2163,11 @@ void SClassViewer::Construct(const FArguments& InArgs, const FClassViewerInitial
 
 	ClassViewerContent = 
 	SNew(SBorder)
-	.BorderImage(FEditorStyle::GetBrush("ToolPanel.GroupBorder"))
+	.BorderImage(FEditorStyle::GetBrush("Menu.Background"))
 	[
 		SNew(SVerticalBox)
 		+SVerticalBox::Slot()
 		.AutoHeight()
-		.Padding( 1.0f, 0.0f, 1.0f, 0.0f )
 		[
 			SNew(SHorizontalBox)
 			+SHorizontalBox::Slot()
