Add missing UObject class in the class hierarchy so BP classes deriving from it appears in the class viewer.
#rnx

diff --git a/Engine/Source/Editor/ClassViewer/Private/SClassViewer.cpp b/Engine/Source/Editor/ClassViewer/Private/SClassViewer.cpp
--- a/Engine/Source/Editor/ClassViewer/Private/SClassViewer.cpp
+++ b/Engine/Source/Editor/ClassViewer/Private/SClassViewer.cpp
@@ -1049,9 +1049,11 @@ void FClassHierarchy::AddChildren_NoFilter( TSharedPtr< FClassViewerNode >& InOu
 {
 	UClass* RootClass = UObject::StaticClass();
 
-	ObjectClassRoot = MakeShareable(new FClassViewerNode(RootClass->GetName(), RootClass->GetDisplayNameText().ToString()));
+	ObjectClassRoot = MakeShareable(new FClassViewerNode(RootClass));
 	ObjectClassRoot->Class = RootClass;
 
+	InOutClassPathToNode.Add(ObjectClassRoot->ClassPath, ObjectClassRoot);
+
 	TMap< UClass*, TSharedPtr< FClassViewerNode > > Nodes;
 
 	Nodes.Add(RootClass, ObjectClassRoot);
