In the ClassViewer widget, make sure that the sorting of the class items obeys the display name type setting, rather than always sorting by class name. This works for both list and tree viewing modes.

The original issue only covered the montage notify list but the fix covers all uses of the class viewer.

diff --git a/Engine/Source/Editor/ClassViewer/Private/SClassViewer.cpp b/Engine/Source/Editor/ClassViewer/Private/SClassViewer.cpp
--- a/Engine/Source/Editor/ClassViewer/Private/SClassViewer.cpp
+++ b/Engine/Source/Editor/ClassViewer/Private/SClassViewer.cpp
@@ -84,20 +84,25 @@ DEFINE_LOG_CATEGORY_STATIC(LogEditorClassViewer, Log, All);
 
 //////////////////////////////////////////////////////////////
 
-// A hack to make IMPLEMENT_COMPARE_CONSTREF with a templated type
-typedef TSharedPtr<FClassViewerNode> FClassViewerNodeSharedPtr;
-
-FORCEINLINE bool CompareFClassViewerNodes( const FClassViewerNodeSharedPtr& A, const FClassViewerNodeSharedPtr& B )
+struct FClassViewerNodeNameLess
 {
-	check(A.IsValid());
-	check(B.IsValid());
+	FClassViewerNodeNameLess(EClassViewerNameTypeToDisplay NameTypeToDisplay = EClassViewerNameTypeToDisplay::ClassName) : NameTypeToDisplay(NameTypeToDisplay) {}
 
-	// Pull out the FString, for ease of reading.
-	FString AString = *A->GetClassName().Get();
-	FString BString = *B->GetClassName().Get();
+	bool operator()(TSharedPtr<FClassViewerNode> A, TSharedPtr<FClassViewerNode> B) const
+	{ 
+		check(A.IsValid());
+		check(B.IsValid());
 
-	return *A->GetClassName().Get() < *B->GetClassName().Get();
-}
+		// The display name only matters when NameTypeToDisplay == DisplayName. For NameTypeToDisplay == Dynamic,
+		// the class name is displayed first with the display name in parentheses, but only if it differs from the display name.
+		bool bUseDisplayName = NameTypeToDisplay == EClassViewerNameTypeToDisplay::DisplayName;
+		const FString& NameA = *A->GetClassName(bUseDisplayName).Get();
+		const FString& NameB = *B->GetClassName(bUseDisplayName).Get();
+		return NameA.Compare(NameB, ESearchCase::IgnoreCase) < 0;
+	}
+
+	EClassViewerNameTypeToDisplay NameTypeToDisplay;
+};
 
 class FClassHierarchy
 {
@@ -112,7 +117,7 @@ public:
 	/** Recursive function to sort a tree.
 	 *	@param InOutRootNode						The current node to sort.
 	 */	
-	void SortChildren( TSharedPtr< FClassViewerNode >& InRootNode );
+	void SortChildren(TSharedPtr< FClassViewerNode >& InRootNode);
 
 	/** Checks if a particular class is placeable.
 	 *	@return The ObjectClassRoot for building a duplicate tree using.
@@ -332,6 +337,11 @@ namespace ClassViewer
 				}
 			}
 
+			if (bReturnPassesFilter)
+			{
+				InOutRootNode->GetChildrenList().Sort(FClassViewerNodeNameLess(InInitOptions.NameTypeToDisplay));
+			}
+
 			return bReturnPassesFilter;
 		}
 
@@ -1299,7 +1309,7 @@ void FClassHierarchy::AddAsset(const FAssetData& InAddedAssetData)
 	}
 }
 
-void FClassHierarchy::SortChildren( TSharedPtr< FClassViewerNode >& InRootNode )
+void FClassHierarchy::SortChildren( TSharedPtr< FClassViewerNode >& InRootNode)
 {
 	TArray< TSharedPtr< FClassViewerNode > >& ChildList = InRootNode->GetChildrenList();
 	for(int32 ChildIndex = 0; ChildIndex < ChildList.Num(); ChildIndex++)
@@ -1312,7 +1322,7 @@ void FClassHierarchy::SortChildren( TSharedPtr< FClassViewerNode >& InRootNode )
 	}
 
 	// Sort the children.
-	ChildList.Sort(CompareFClassViewerNodes);
+	ChildList.Sort(FClassViewerNodeNameLess());
 }
 
 void FClassHierarchy::FindClass(TSharedPtr< FClassViewerNode > InOutClassNode)
@@ -2476,19 +2486,7 @@ void SClassViewer::Populate()
 		ClassViewer::Helpers::GetClassList(RootTreeItems, ClassFilter, InitOptions);
 
 		// Sort the list alphabetically.
-		struct FCompareFClassViewerNode
-		{
-			FORCEINLINE bool operator()( TSharedPtr<FClassViewerNode> A, TSharedPtr<FClassViewerNode> B ) const
-			{
-				check(A.IsValid());
-				check(B.IsValid());
-				// Pull out the FString, for ease of reading.
-				const FString& AString = *A->GetClassName().Get();
-				const FString& BString = *B->GetClassName().Get();
-				return AString < BString;
-			}
-		};
-		RootTreeItems.Sort( FCompareFClassViewerNode() );
+		RootTreeItems.Sort(FClassViewerNodeNameLess(InitOptions.NameTypeToDisplay));
 
 		// Only display this option if the user wants it and in Picker Mode.
 		if(InitOptions.bShowNoneOption && InitOptions.Mode == EClassViewerMode::ClassPicker)
