UClassViewerSettings now contains an array of allowed classes which can be used to filter out classes from SClassViewer instances, such as when picking classes in the Details view or creating blueprints. This list is currently defined in EditorPerProjectUserSettings.ini.

diff --git a/Engine/Source/Editor/ClassViewer/Private/ClassViewerFilter.cpp b/Engine/Source/Editor/ClassViewer/Private/ClassViewerFilter.cpp
--- a/Engine/Source/Editor/ClassViewer/Private/ClassViewerFilter.cpp
+++ b/Engine/Source/Editor/ClassViewer/Private/ClassViewerFilter.cpp
@@ -448,10 +448,12 @@ bool FClassViewerFilter::IsClassAllowed(const FClassViewerInitializationOptions&
 	// Determine if we allow any developer folder classes, if so determine if this class is in one of the allowed developer folders.
 	static const FString DeveloperPathWithSlash = FPackageName::FilenameToLongPackageName(FPaths::GameDevelopersDir());
 	static const FString UserDeveloperPathWithSlash = FPackageName::FilenameToLongPackageName(FPaths::GameUserDeveloperDir());
-	FString GeneratedClassPathString = InClass->GetPathName();
+	const FString GeneratedClassPathString = InClass->GetPathName();
+
+	const UClassViewerSettings* ClassViewerSettings = GetDefault<UClassViewerSettings>();
 
 	bool bPassesDeveloperFilter = true;
-	EClassViewerDeveloperType AllowedDeveloperType = GetDefault<UClassViewerSettings>()->DeveloperFolderType;
+	EClassViewerDeveloperType AllowedDeveloperType = ClassViewerSettings->DeveloperFolderType;
 	if (AllowedDeveloperType == EClassViewerDeveloperType::CVDT_None)
 	{
 		bPassesDeveloperFilter = !GeneratedClassPathString.StartsWith(DeveloperPathWithSlash);
@@ -463,15 +465,15 @@ bool FClassViewerFilter::IsClassAllowed(const FClassViewerInitializationOptions&
 			bPassesDeveloperFilter = GeneratedClassPathString.StartsWith(UserDeveloperPathWithSlash);
 		}
 	}
-
+	
 	// The INI files declare classes and folders that are considered internal only. Does this class match any of those patterns?
 	// INI path: /Script/ClassViewer.ClassViewerProjectSettings
 	bool bPassesInternalFilter = true;
-	if (!GetDefault<UClassViewerSettings>()->DisplayInternalClasses)
+	if (!ClassViewerSettings->DisplayInternalClasses)
 	{
-		for (int i = 0; i < InternalPaths.Num(); ++i)
+		for (const FDirectoryPath& DirPath : InternalPaths)
 		{
-			if (GeneratedClassPathString.StartsWith(InternalPaths[i].Path))
+			if (GeneratedClassPathString.StartsWith(DirPath.Path))
 			{
 				bPassesInternalFilter = false;
 				break;
@@ -480,9 +482,9 @@ bool FClassViewerFilter::IsClassAllowed(const FClassViewerInitializationOptions&
 
 		if (bPassesInternalFilter)
 		{
-			for (int i = 0; i < InternalClasses.Num(); ++i)
+			for (const UClass* Class : InternalClasses)
 			{
-				if (InClass->IsChildOf(InternalClasses[i]))
+				if (InClass->IsChildOf(Class))
 				{
 					bPassesInternalFilter = false;
 					break;
@@ -491,6 +493,16 @@ bool FClassViewerFilter::IsClassAllowed(const FClassViewerInitializationOptions&
 		}
 	}
 
+	// The INI files can contain a list of globally allowed classes - if it does, then only classes whose names match will be shown.
+	bool bPassesAllowedClasses = true;
+	if (ClassViewerSettings->AllowedClasses.Num() > 0)
+	{
+		if (!ClassViewerSettings->AllowedClasses.Contains(GeneratedClassPathString))
+		{
+			bPassesAllowedClasses = false;
+		}
+	}
+
 	bool bPassesPlaceableFilter = true;
 	if (InInitOptions.bIsPlaceableOnly)
 	{
@@ -512,7 +524,7 @@ bool FClassViewerFilter::IsClassAllowed(const FClassViewerInitializationOptions&
 		bPassesAssetReferenceFilter = AssetReferenceFilter->PassesFilter(FAssetData(InClass));
 	}
 
-	bool bPassesFilter = bPassesPlaceableFilter && bPassesBlueprintBaseFilter 
+	bool bPassesFilter = bPassesAllowedClasses && bPassesPlaceableFilter && bPassesBlueprintBaseFilter
 		&& bPassesDeveloperFilter && bPassesInternalFilter && bPassesEditorClassFilter 
 		&& bPassesCustomFilter && (!bCheckTextFilter || bPassesTextFilter) && bPassesAssetReferenceFilter;
 
@@ -545,7 +557,8 @@ bool FClassViewerFilter::IsUnloadedClassAllowed(const FClassViewerInitialization
 
 	bool bPassesDeveloperFilter = true;
 
-	EClassViewerDeveloperType AllowedDeveloperType = GetDefault<UClassViewerSettings>()->DeveloperFolderType;
+	const UClassViewerSettings* ClassViewerSettings = GetDefault<UClassViewerSettings>();
+	EClassViewerDeveloperType AllowedDeveloperType = ClassViewerSettings->DeveloperFolderType;
 	if (AllowedDeveloperType == EClassViewerDeveloperType::CVDT_None)
 	{
 		bPassesDeveloperFilter = !GeneratedClassPathString.StartsWith(DeveloperPathWithSlash);
@@ -561,11 +574,11 @@ bool FClassViewerFilter::IsUnloadedClassAllowed(const FClassViewerInitialization
 	// The INI files declare classes and folders that are considered internal only. Does this class match any of those patterns?
 	// INI path: /Script/ClassViewer.ClassViewerProjectSettings
 	bool bPassesInternalFilter = true;
-	if (!GetDefault<UClassViewerSettings>()->DisplayInternalClasses)
+	if (!ClassViewerSettings->DisplayInternalClasses)
 	{
-		for (int i = 0; i < InternalPaths.Num(); ++i)
+		for (const FDirectoryPath& DirPath : InternalPaths)
 		{
-			if (GeneratedClassPathString.StartsWith(InternalPaths[i].Path))
+			if (GeneratedClassPathString.StartsWith(DirPath.Path))
 			{
 				bPassesInternalFilter = false;
 				break;
@@ -573,6 +586,16 @@ bool FClassViewerFilter::IsUnloadedClassAllowed(const FClassViewerInitialization
 		}
 	}
 
+	// The INI files can contain a list of globally allowed classes - if it does, then only classes whose names match will be shown.
+	bool bPassesAllowedClasses = true;
+	if (ClassViewerSettings->AllowedClasses.Num() > 0)
+	{
+		if (!ClassViewerSettings->AllowedClasses.Contains(GeneratedClassPathString))
+		{
+			bPassesAllowedClasses = false;
+		}
+	}
+
 	bool bPassesPlaceableFilter = true;
 	if (InInitOptions.bIsPlaceableOnly)
 	{
diff --git a/Engine/Source/Editor/UnrealEd/Classes/Settings/ClassViewerSettings.h b/Engine/Source/Editor/UnrealEd/Classes/Settings/ClassViewerSettings.h
--- a/Engine/Source/Editor/UnrealEd/Classes/Settings/ClassViewerSettings.h
+++ b/Engine/Source/Editor/UnrealEd/Classes/Settings/ClassViewerSettings.h
@@ -35,6 +35,9 @@ class UNREALED_API UClassViewerSettings
 
 public:
 
+	UPROPERTY(config)
+	TArray<FString> AllowedClasses;
+
 	/** Whether to display internal use classes. */
 	UPROPERTY(config)
 	bool DisplayInternalClasses;
@@ -52,13 +55,12 @@ public:
 	 * @return The delegate.
 	 */
 	DECLARE_EVENT_OneParam(UClassViewerSettings, FSettingChangedEvent, FName /*PropertyName*/);
-	static FSettingChangedEvent& OnSettingChanged( ) { return SettingChangedEvent; }
+	static FSettingChangedEvent& OnSettingChanged() { return SettingChangedEvent; }
 
 protected:
 
 	// UObject overrides
-
-	virtual void PostEditChangeProperty( struct FPropertyChangedEvent& PropertyChangedEvent ) override;
+	virtual void PostEditChangeProperty(struct FPropertyChangedEvent& PropertyChangedEvent) override;
 
 
 	// Holds an event delegate that is executed when a setting has changed.
