Fixed developer folder filter in SClassViewer

diff --git a/Engine/Source/Editor/ClassViewer/Private/SClassViewer.cpp b/Engine/Source/Editor/ClassViewer/Private/SClassViewer.cpp
--- a/Engine/Source/Editor/ClassViewer/Private/SClassViewer.cpp
+++ b/Engine/Source/Editor/ClassViewer/Private/SClassViewer.cpp
@@ -855,7 +855,7 @@ namespace ClassViewer
 			FString GeneratedClassPathString = InOriginalRootNode->ClassPath.ToString();
 			if (InAllowedDeveloperType == EClassViewerDeveloperType::CVDT_None)
 			{
-				bPassesDeveloperFilter = GeneratedClassPathString.StartsWith(DeveloperPathWithSlash);
+				bPassesDeveloperFilter = !GeneratedClassPathString.StartsWith(DeveloperPathWithSlash);
 			}
 			else if (InAllowedDeveloperType == EClassViewerDeveloperType::CVDT_CurrentUser)
 			{
