SPropertyEditorClass now checks that a dropped class passes the class filter, preventing a crash if an invalid class is dropped on it.

To accomplish this, it is now possible to retrieve an IClassViewerFilter from the ClassViewer module.

In addition, most checks that SClassViewer was doing in static methods have been moved to FClassViewerFilter.

diff --git a/Engine/Source/Editor/ClassViewer/Private/ClassViewerFilter.cpp b/Engine/Source/Editor/ClassViewer/Private/ClassViewerFilter.cpp
new file mode 100644
--- /dev/null
+++ b/Engine/Source/Editor/ClassViewer/Private/ClassViewerFilter.cpp
@@ -0,0 +1,552 @@
+#include "ClassViewerFilter.h"
+
+#include "UnloadedBlueprintData.h"
+#include "Kismet2/KismetEditorUtilities.h"
+#include "Misc/TextFilterExpressionEvaluator.h"
+
+EFilterReturn::Type FClassViewerFilterFuncs::IfInChildOfClassesSet(TSet< const UClass* >& InSet, const UClass* InClass)
+{
+	check(InClass);
+
+	if (InSet.Num())
+	{
+		// If a class is a child of any classes on this list, it will be allowed onto the list, unless it also appears on a disallowed list.
+		for (auto CurClassIt = InSet.CreateConstIterator(); CurClassIt; ++CurClassIt)
+		{
+			if (InClass->IsChildOf(*CurClassIt))
+			{
+				return EFilterReturn::Passed;
+			}
+		}
+
+		return EFilterReturn::Failed;
+	}
+
+	// Since there are none on this list, return that there is no items.
+	return EFilterReturn::NoItems;
+}
+
+EFilterReturn::Type FClassViewerFilterFuncs::IfInChildOfClassesSet(TSet< const UClass* >& InSet, const TSharedPtr< const IUnloadedBlueprintData > InClass)
+{
+	check(InClass.IsValid());
+
+	if (InSet.Num())
+	{
+		// If a class is a child of any classes on this list, it will be allowed onto the list, unless it also appears on a disallowed list.
+		for (auto CurClassIt = InSet.CreateConstIterator(); CurClassIt; ++CurClassIt)
+		{
+			if (InClass->IsChildOf(*CurClassIt))
+			{
+				return EFilterReturn::Passed;
+			}
+		}
+
+		return EFilterReturn::Failed;
+	}
+
+	// Since there are none on this list, return that there is no items.
+	return EFilterReturn::NoItems;
+}
+
+EFilterReturn::Type FClassViewerFilterFuncs::IfMatchesAllInChildOfClassesSet(TSet< const UClass* >& InSet, const UClass* InClass)
+{
+	check(InClass);
+
+	if (InSet.Num())
+	{
+		// If a class is a child of any classes on this list, it will be allowed onto the list, unless it also appears on a disallowed list.
+		for (auto CurClassIt = InSet.CreateConstIterator(); CurClassIt; ++CurClassIt)
+		{
+			if (!InClass->IsChildOf(*CurClassIt))
+			{
+				// Since it doesn't match one, it fails.
+				return EFilterReturn::Failed;
+			}
+		}
+
+		// It matches all of them, so it passes.
+		return EFilterReturn::Passed;
+	}
+
+	// Since there are none on this list, return that there is no items.
+	return EFilterReturn::NoItems;
+}
+
+EFilterReturn::Type FClassViewerFilterFuncs::IfMatchesAllInChildOfClassesSet(TSet< const UClass* >& InSet, const TSharedPtr< const IUnloadedBlueprintData > InClass)
+{
+	check(InClass.IsValid());
+
+	if (InSet.Num())
+	{
+		// If a class is a child of any classes on this list, it will be allowed onto the list, unless it also appears on a disallowed list.
+		for (auto CurClassIt = InSet.CreateConstIterator(); CurClassIt; ++CurClassIt)
+		{
+			if (!InClass->IsChildOf(*CurClassIt))
+			{
+				// Since it doesn't match one, it fails.
+				return EFilterReturn::Failed;
+			}
+		}
+
+		// It matches all of them, so it passes.
+		return EFilterReturn::Passed;
+	}
+
+	// Since there are none on this list, return that there is no items.
+	return EFilterReturn::NoItems;
+}
+
+EFilterReturn::Type FClassViewerFilterFuncs::IfMatchesAll_ObjectsSetIsAClass(TSet< const UObject* >& InSet, const UClass* InClass)
+{
+	check(InClass);
+
+	if (InSet.Num())
+	{
+		// If a class is a child of any classes on this list, it will be allowed onto the list, unless it also appears on a disallowed list.
+		for (auto CurClassIt = InSet.CreateConstIterator(); CurClassIt; ++CurClassIt)
+		{
+			if (!(*CurClassIt)->IsA(InClass))
+			{
+				// Since it doesn't match one, it fails.
+				return EFilterReturn::Failed;
+			}
+		}
+
+		// It matches all of them, so it passes.
+		return EFilterReturn::Passed;
+	}
+
+	// Since there are none on this list, return that there is no items.
+	return EFilterReturn::NoItems;
+}
+
+EFilterReturn::Type FClassViewerFilterFuncs::IfMatchesAll_ObjectsSetIsAClass(TSet< const UObject* >& InSet, const TSharedPtr< const IUnloadedBlueprintData > InClass)
+{
+	check(InClass.IsValid());
+
+	if (InSet.Num())
+	{
+		// If a class is a child of any classes on this list, it will be allowed onto the list, unless it also appears on a disallowed list.
+		for (auto CurClassIt = InSet.CreateConstIterator(); CurClassIt; ++CurClassIt)
+		{
+			if (!(*CurClassIt)->IsA(UBlueprintGeneratedClass::StaticClass()))
+			{
+				// Since it doesn't match one, it fails.
+				return EFilterReturn::Failed;
+			}
+		}
+
+		// It matches all of them, so it passes.
+		return EFilterReturn::Passed;
+	}
+
+	// Since there are none on this list, return that there is no items.
+	return EFilterReturn::NoItems;
+}
+
+EFilterReturn::Type FClassViewerFilterFuncs::IfMatchesAll_ClassesSetIsAClass(TSet< const UClass* >& InSet, const UClass* InClass)
+{
+	check(InClass);
+
+	if (InSet.Num())
+	{
+		// If a class is a child of any classes on this list, it will be allowed onto the list, unless it also appears on a disallowed list.
+		for (auto CurClassIt = InSet.CreateConstIterator(); CurClassIt; ++CurClassIt)
+		{
+			const UObject* Object = *CurClassIt;
+			if (!Object->IsA(InClass))
+			{
+				// Since it doesn't match one, it fails.
+				return EFilterReturn::Failed;
+			}
+		}
+
+		// It matches all of them, so it passes.
+		return EFilterReturn::Passed;
+	}
+
+	// Since there are none on this list, return that there is no items.
+	return EFilterReturn::NoItems;
+}
+
+EFilterReturn::Type FClassViewerFilterFuncs::IfMatchesAll_ClassesSetIsAClass(TSet< const UClass* >& InSet, const TSharedPtr< const IUnloadedBlueprintData > InClass)
+{
+	check(InClass.IsValid());
+
+	if (InSet.Num())
+	{
+		// If a class is a child of any classes on this list, it will be allowed onto the list, unless it also appears on a disallowed list.
+		for (auto CurClassIt = InSet.CreateConstIterator(); CurClassIt; ++CurClassIt)
+		{
+			const UObject* Object = *CurClassIt;
+			if (!Object->IsA(UBlueprintGeneratedClass::StaticClass()))
+			{
+				// Since it doesn't match one, it fails.
+				return EFilterReturn::Failed;
+			}
+		}
+
+		// It matches all of them, so it passes.
+		return EFilterReturn::Passed;
+	}
+
+	// Since there are none on this list, return that there is no items.
+	return EFilterReturn::NoItems;
+}
+
+EFilterReturn::Type FClassViewerFilterFuncs::IfMatches_ClassesSetIsAClass(TSet< const UClass* >& InSet, const UClass* InClass)
+{
+	check(InClass);
+
+	if (InSet.Num())
+	{
+		// If a class is a child of any classes on this list, it will be allowed onto the list, unless it also appears on a disallowed list.
+		for (auto CurClassIt = InSet.CreateConstIterator(); CurClassIt; ++CurClassIt)
+		{
+			const UObject* Object = *CurClassIt;
+			if (Object->IsA(InClass))
+			{
+				return EFilterReturn::Passed;
+			}
+		}
+
+		return EFilterReturn::Failed;
+	}
+
+	// Since there are none on this list, return that there is no items.
+	return EFilterReturn::NoItems;
+}
+
+EFilterReturn::Type FClassViewerFilterFuncs::IfMatches_ClassesSetIsAClass(TSet< const UClass* >& InSet, const TSharedPtr< const IUnloadedBlueprintData > InClass)
+{
+	check(InClass.IsValid());
+
+	if (InSet.Num())
+	{
+		// If a class is a child of any classes on this list, it will be allowed onto the list, unless it also appears on a disallowed list.
+		for (auto CurClassIt = InSet.CreateConstIterator(); CurClassIt; ++CurClassIt)
+		{
+			const UObject* Object = *CurClassIt;
+			if (Object->IsA(UBlueprintGeneratedClass::StaticClass()))
+			{
+				return EFilterReturn::Passed;
+			}
+		}
+
+		return EFilterReturn::Failed;
+	}
+
+	// Since there are none on this list, return that there is no items.
+	return EFilterReturn::NoItems;
+}
+
+EFilterReturn::Type FClassViewerFilterFuncs::IfInClassesSet(TSet< const UClass* >& InSet, const UClass* InClass)
+{
+	check(InClass);
+
+	if (InSet.Num())
+	{
+		for (auto CurClassIt = InSet.CreateConstIterator(); CurClassIt; ++CurClassIt)
+		{
+			if (InClass == *CurClassIt)
+			{
+				return EFilterReturn::Passed;
+			}
+		}
+		return EFilterReturn::Failed;
+	}
+
+	// Since there are none on this list, return that there is no items.
+	return EFilterReturn::NoItems;
+}
+
+EFilterReturn::Type FClassViewerFilterFuncs::IfInClassesSet(TSet< const UClass* >& InSet, const TSharedPtr< const IUnloadedBlueprintData > InClass)
+{
+	check(InClass.IsValid());
+
+	if (InSet.Num())
+	{
+		for (auto CurClassIt = InSet.CreateConstIterator(); CurClassIt; ++CurClassIt)
+		{
+			const TSharedPtr<const FUnloadedBlueprintData> UnloadedBlueprintData = StaticCastSharedPtr<const FUnloadedBlueprintData>(InClass);
+			if (*UnloadedBlueprintData->GetClassViewerNode().Pin()->GetClassName() == (*CurClassIt)->GetName())
+			{
+				return EFilterReturn::Passed;
+			}
+		}
+		return EFilterReturn::Failed;
+	}
+
+	// Since there are none on this list, return that there is no items.
+	return EFilterReturn::NoItems;
+}
+
+/** Checks if a particular class is a brush.
+*	@param InClass				The Class to check.
+*	@return Returns true if the class is a brush.
+*/
+static bool IsBrush(const UClass* InClass)
+{
+	return InClass->IsChildOf(ABrush::StaticClass());
+}
+
+static bool IsBrush(const TSharedRef<const IUnloadedBlueprintData>& InBlueprintData)
+{
+	return InBlueprintData->IsChildOf(ABrush::StaticClass());
+}
+
+/** Checks if a particular class is placeable.
+*	@param InClass				The Class to check.
+*	@return Returns true if the class is placeable.
+*/
+static bool IsPlaceable(const UClass* InClass)
+{
+	return !InClass->HasAnyClassFlags(CLASS_Abstract | CLASS_NotPlaceable) 
+		&& InClass->IsChildOf(AActor::StaticClass());
+}
+
+static bool IsPlaceable(const TSharedRef<const IUnloadedBlueprintData>& InBlueprintData)
+{
+	return !InBlueprintData->HasAnyClassFlags(CLASS_Abstract | CLASS_NotPlaceable)
+		&& InBlueprintData->IsChildOf(AActor::StaticClass());
+}
+
+/** Util class to checks if a particular class can be made into a Blueprint, ignores deprecation
+ *
+ * @param InClass	The class to verify can be made into a Blueprint
+ * @return			true if the class can be made into a Blueprint
+ */
+bool CanCreateBlueprintOfClass_IgnoreDeprecation(UClass* InClass)
+{
+	// Temporarily remove the deprecated flag so we can check if it is valid for
+	bool bIsClassDeprecated = InClass->HasAnyClassFlags(CLASS_Deprecated);
+	InClass->ClassFlags &= ~CLASS_Deprecated;
+
+	bool bCanCreateBlueprintOfClass = FKismetEditorUtilities::CanCreateBlueprintOfClass(InClass);
+
+	// Reassign the deprecated flag if it was previously assigned
+	if (bIsClassDeprecated)
+	{
+		InClass->ClassFlags |= CLASS_Deprecated;
+	}
+
+	return bCanCreateBlueprintOfClass;
+}
+
+/** Checks if a node is a blueprint base or not.
+*	@param	InNode	The node to check if it is a blueprint base.
+*	@return			true if the class is a blueprint base.
+*/
+static bool CheckIfBlueprintBase(const TSharedRef<const IUnloadedBlueprintData>& InBlueprintData)
+{
+	if (InBlueprintData->IsNormalBlueprintType())
+	{
+		bool bAllowDerivedBlueprints = false;
+		GConfig->GetBool(TEXT("Kismet"), TEXT("AllowDerivedBlueprints"), /*out*/ bAllowDerivedBlueprints, GEngineIni);
+
+		return bAllowDerivedBlueprints;
+	}
+	return false;
+}
+
+/** Checks if the TestString passes the filter.
+*	@param InTestString			The string to test against the filter.
+*	@param InTextFilter			Compiled text filter to apply.
+*
+*	@return	true if it passes the filter.
+*/
+static bool PassesTextFilter(const FString& InTestString, const TSharedRef<FTextFilterExpressionEvaluator>& InTextFilter)
+{
+	class FClassFilterContext : public ITextFilterExpressionContext
+	{
+	public:
+		explicit FClassFilterContext(const FString& InStr)
+			: StrPtr(&InStr)
+		{
+		}
+
+		virtual bool TestBasicStringExpression(const FTextFilterString& InValue, const ETextFilterTextComparisonMode InTextComparisonMode) const override
+		{
+			return TextFilterUtils::TestBasicStringExpression(*StrPtr, InValue, InTextComparisonMode);
+		}
+
+		virtual bool TestComplexExpression(const FName& InKey, const FTextFilterString& InValue, const ETextFilterComparisonOperation InComparisonOperation, const ETextFilterTextComparisonMode InTextComparisonMode) const override
+		{
+			return false;
+		}
+
+	private:
+		const FString* StrPtr;
+	};
+
+	return InTextFilter->TestTextFilter(FClassFilterContext(InTestString));
+}
+
+FClassViewerFilter::FClassViewerFilter(const FClassViewerInitializationOptions& InInitOptions) :
+	TextFilter(MakeShared<FTextFilterExpressionEvaluator>(ETextFilterExpressionEvaluatorMode::BasicString)),
+	FilterFunctions(MakeShared<FClassViewerFilterFuncs>())
+{
+}
+
+bool FClassViewerFilter::IsNodeAllowed(const FClassViewerInitializationOptions& InInitOptions, const TSharedRef<FClassViewerNode>& InNode)
+{
+	if (InNode->Class.IsValid())
+	{
+		return IsClassAllowed(InInitOptions, InNode->Class.Get(), FilterFunctions);
+	}
+	else if (InInitOptions.bShowUnloadedBlueprints && InNode->UnloadedBlueprintData.IsValid())
+	{
+		return IsUnloadedClassAllowed(InInitOptions, InNode->UnloadedBlueprintData.ToSharedRef(), FilterFunctions);
+	}
+
+	return false;
+}
+
+PRAGMA_DISABLE_OPTIMIZATION
+
+bool FClassViewerFilter::IsClassAllowed(const FClassViewerInitializationOptions& InInitOptions, const UClass* InClass, TSharedRef<class FClassViewerFilterFuncs> InFilterFuncs)
+{
+	if (InInitOptions.bIsActorsOnly && !InClass->IsChildOf(AActor::StaticClass()))
+	{
+		return false;
+	}
+
+	const bool bPassesBlueprintBaseFilter = !InInitOptions.bIsBlueprintBaseOnly || CanCreateBlueprintOfClass_IgnoreDeprecation(const_cast<UClass*>(InClass));
+	const bool bPassesEditorClassFilter = !InInitOptions.bEditorClassesOnly || IsEditorOnlyObject(InClass);
+
+	// Determine if we allow any developer folder classes, if so determine if this class is in one of the allowed developer folders.
+	static const FString DeveloperPathWithSlash = FPackageName::FilenameToLongPackageName(FPaths::GameDevelopersDir());
+	static const FString UserDeveloperPathWithSlash = FPackageName::FilenameToLongPackageName(FPaths::GameUserDeveloperDir());
+	FString GeneratedClassPathString = InClass->GetPathName();
+
+	bool bPassesDeveloperFilter = true;
+	EClassViewerDeveloperType AllowedDeveloperType = GetDefault<UClassViewerSettings>()->DeveloperFolderType;
+	if (AllowedDeveloperType == EClassViewerDeveloperType::CVDT_None)
+	{
+		bPassesDeveloperFilter = !GeneratedClassPathString.StartsWith(DeveloperPathWithSlash);
+	}
+	else if (AllowedDeveloperType == EClassViewerDeveloperType::CVDT_CurrentUser)
+	{
+		if (GeneratedClassPathString.StartsWith(DeveloperPathWithSlash))
+		{
+			bPassesDeveloperFilter = GeneratedClassPathString.StartsWith(UserDeveloperPathWithSlash);
+		}
+	}
+
+	// The INI files declare classes and folders that are considered internal only. Does this class match any of those patterns?
+	// INI path: /Script/ClassViewer.ClassViewerProjectSettings
+	bool bPassesInternalFilter = true;
+	if (!GetDefault<UClassViewerSettings>()->DisplayInternalClasses)
+	{
+		for (int i = 0; i < InternalPaths.Num(); ++i)
+		{
+			if (GeneratedClassPathString.StartsWith(InternalPaths[i].Path))
+			{
+				bPassesInternalFilter = false;
+				break;
+			}
+		}
+
+		if (bPassesInternalFilter)
+		{
+			for (int i = 0; i < InternalClasses.Num(); ++i)
+			{
+				if (InClass->IsChildOf(InternalClasses[i]))
+				{
+					bPassesInternalFilter = false;
+					break;
+				}
+			}
+		}
+	}
+
+	bool bPassesPlaceableFilter = true;
+	if (InInitOptions.bIsPlaceableOnly)
+	{
+		bPassesPlaceableFilter = IsPlaceable(InClass) && 
+			(InInitOptions.Mode == EClassViewerMode::ClassPicker || !IsBrush(InClass));
+	}
+
+	bool bPassesCustomFilter = true;
+	if (InInitOptions.ClassFilter.IsValid())
+	{
+		bPassesCustomFilter = InInitOptions.ClassFilter->IsClassAllowed(InInitOptions, InClass, FilterFunctions);
+	}
+
+	const bool bPassesTextFilter = PassesTextFilter(InClass->GetName(), TextFilter);
+
+	bool bPassesFilter = bPassesPlaceableFilter && bPassesBlueprintBaseFilter 
+		&& bPassesDeveloperFilter && bPassesInternalFilter && bPassesEditorClassFilter 
+		&& bPassesCustomFilter && bPassesTextFilter;
+
+	return bPassesFilter;
+}
+
+bool FClassViewerFilter::IsUnloadedClassAllowed(const FClassViewerInitializationOptions& InInitOptions, const TSharedRef<const IUnloadedBlueprintData> InUnloadedClassData, TSharedRef<FClassViewerFilterFuncs> InFilterFuncs)
+{
+	if (InInitOptions.bIsActorsOnly && !InUnloadedClassData->IsChildOf(AActor::StaticClass()))
+	{
+		return false;
+	}
+
+	const bool bPassesBlueprintBaseFilter = !InInitOptions.bIsBlueprintBaseOnly || CheckIfBlueprintBase(InUnloadedClassData);
+
+	// unloaded blueprints cannot be editor-only
+	const bool bPassesEditorClassFilter = !InInitOptions.bEditorClassesOnly;
+
+	// Determine if we allow any developer folder classes, if so determine if this class is in one of the allowed developer folders.
+	static const FString DeveloperPathWithSlash = FPackageName::FilenameToLongPackageName(FPaths::GameDevelopersDir());
+	static const FString UserDeveloperPathWithSlash = FPackageName::FilenameToLongPackageName(FPaths::GameUserDeveloperDir());
+	FString GeneratedClassPathString = InUnloadedClassData->GetClassPath().ToString();
+
+	bool bPassesDeveloperFilter = true;
+
+	EClassViewerDeveloperType AllowedDeveloperType = GetDefault<UClassViewerSettings>()->DeveloperFolderType;
+	if (AllowedDeveloperType == EClassViewerDeveloperType::CVDT_None)
+	{
+		bPassesDeveloperFilter = !GeneratedClassPathString.StartsWith(DeveloperPathWithSlash);
+	}
+	else if (AllowedDeveloperType == EClassViewerDeveloperType::CVDT_CurrentUser)
+	{
+		if (GeneratedClassPathString.StartsWith(DeveloperPathWithSlash))
+		{
+			bPassesDeveloperFilter = GeneratedClassPathString.StartsWith(UserDeveloperPathWithSlash);
+		}
+	}
+
+	// The INI files declare classes and folders that are considered internal only. Does this class match any of those patterns?
+	// INI path: /Script/ClassViewer.ClassViewerProjectSettings
+	bool bPassesInternalFilter = true;
+	if (!GetDefault<UClassViewerSettings>()->DisplayInternalClasses)
+	{
+		for (int i = 0; i < InternalPaths.Num(); ++i)
+		{
+			if (GeneratedClassPathString.StartsWith(InternalPaths[i].Path))
+			{
+				bPassesInternalFilter = false;
+				break;
+			}
+		}
+	}
+
+	bool bPassesPlaceableFilter = true;
+	if (InInitOptions.bIsPlaceableOnly)
+	{
+		bPassesPlaceableFilter = IsPlaceable(InUnloadedClassData) &&
+			(InInitOptions.Mode == EClassViewerMode::ClassPicker || !IsBrush(InUnloadedClassData));
+	}
+
+	bool bPassesCustomFilter = true;
+	if (InInitOptions.ClassFilter.IsValid())
+	{
+		bPassesCustomFilter = InInitOptions.ClassFilter->IsUnloadedClassAllowed(InInitOptions, InUnloadedClassData, FilterFunctions);
+	}
+
+	const bool bPassesTextFilter = PassesTextFilter(*InUnloadedClassData->GetClassName().Get(), TextFilter);
+
+	bool bPassesFilter = bPassesPlaceableFilter && bPassesBlueprintBaseFilter
+		&& bPassesDeveloperFilter && bPassesInternalFilter && bPassesEditorClassFilter
+		&& bPassesCustomFilter && bPassesTextFilter;
+
+	return bPassesFilter;
+}
diff --git a/Engine/Source/Editor/ClassViewer/Private/ClassViewerModule.cpp b/Engine/Source/Editor/ClassViewer/Private/ClassViewerModule.cpp
--- a/Engine/Source/Editor/ClassViewer/Private/ClassViewerModule.cpp
+++ b/Engine/Source/Editor/ClassViewer/Private/ClassViewerModule.cpp
@@ -16,6 +16,7 @@
 #include "Widgets/Docking/SDockTab.h"
 #include "ClassViewerProjectSettings.h"
 #include "ISettingsModule.h"
+#include "ClassViewerFilter.h"
 
 #define LOCTEXT_NAMESPACE "ClassViewer"
 
@@ -95,5 +96,14 @@ TSharedRef<SWidget> FClassViewerModule::CreateClassViewer(const FClassViewerInit
 			.OnClassPickedDelegate(OnClassPickedDelegate);
 }
 
+TSharedRef<IClassViewerFilter> FClassViewerModule::CreateClassFilter(const FClassViewerInitializationOptions& InitOptions)
+{
+	return TSharedRef<IClassViewerFilter>(new FClassViewerFilter(InitOptions));
+}
+
+TSharedRef<FClassViewerFilterFuncs> FClassViewerModule::CreateFilterFuncs()
+{
+	return TSharedRef<FClassViewerFilterFuncs>(new FClassViewerFilterFuncs());
+}
 
 #undef LOCTEXT_NAMESPACE
diff --git a/Engine/Source/Editor/ClassViewer/Private/ClassViewerNode.cpp b/Engine/Source/Editor/ClassViewer/Private/ClassViewerNode.cpp
--- a/Engine/Source/Editor/ClassViewer/Private/ClassViewerNode.cpp
+++ b/Engine/Source/Editor/ClassViewer/Private/ClassViewerNode.cpp
@@ -8,12 +8,35 @@
 #include "ClassViewerFilter.h"
 #include "PropertyHandle.h"
 
+FClassViewerNode::FClassViewerNode(UClass* InClass)
+{
+	Class = InClass;
+	ClassName = MakeShareable(new FString(Class->GetName()));
+	ClassDisplayName = MakeShareable(new FString(Class->GetDisplayNameText().ToString()));
+	ClassPath = FName(*Class->GetPathName());
+
+	if (Class->GetSuperClass())
+	{
+		ParentClassPath = FName(*Class->GetSuperClass()->GetPathName());
+	}
+
+	if (Class->ClassGeneratedBy && Class->ClassGeneratedBy->IsA(UBlueprint::StaticClass()))
+	{
+		Blueprint = Cast<UBlueprint>(Class->ClassGeneratedBy);
+	}
+	else
+	{
+		Blueprint = nullptr;
+	}
+
+	bPassesFilter = false;
+}
+
 FClassViewerNode::FClassViewerNode(const FString& InClassName, const FString& InClassDisplayName)
 {
 	ClassName = MakeShareable(new FString(InClassName));
 	ClassDisplayName = MakeShareable(new FString(InClassDisplayName));
 	bPassesFilter = false;
-	bIsBPNormalType = false;
 
 	Class = nullptr;
 	Blueprint = nullptr;
@@ -34,7 +57,6 @@ FClassViewerNode::FClassViewerNode( const FClassViewerNode& InCopyObject)
 	ParentClassPath = InCopyObject.ParentClassPath;
 	ClassName = InCopyObject.ClassName;
 	BlueprintAssetPath = InCopyObject.BlueprintAssetPath;
-	bIsBPNormalType = InCopyObject.bIsBPNormalType;
 
 	// We do not want to copy the child list, do not add it. It should be the only item missing.
 }
diff --git a/Engine/Source/Editor/ClassViewer/Private/ClassViewerNode.h b/Engine/Source/Editor/ClassViewer/Private/ClassViewerNode.h
--- a/Engine/Source/Editor/ClassViewer/Private/ClassViewerNode.h
+++ b/Engine/Source/Editor/ClassViewer/Private/ClassViewerNode.h
@@ -13,6 +13,8 @@ class UBlueprint;
 class FClassViewerNode
 {
 public:
+	FClassViewerNode(UClass* Class);
+
 	/**
 	 * Creates a node for the widget's tree.
 	 *
@@ -101,9 +103,6 @@ public:
 	/** true if the class passed the filter. */
 	bool bPassesFilter;
 
-	/** true if the class is a "normal type", this is used to identify unloaded blueprints as blueprint bases. */
-	bool bIsBPNormalType;
-
 	/** Pointer to the parent to this object. */
 	TWeakPtr< FClassViewerNode > ParentNode;
 
diff --git a/Engine/Source/Editor/ClassViewer/Private/SClassViewer.cpp b/Engine/Source/Editor/ClassViewer/Private/SClassViewer.cpp
--- a/Engine/Source/Editor/ClassViewer/Private/SClassViewer.cpp
+++ b/Engine/Source/Editor/ClassViewer/Private/SClassViewer.cpp
@@ -81,287 +81,6 @@
 
 DEFINE_LOG_CATEGORY_STATIC(LogEditorClassViewer, Log, All);
 
-//////////////////////////////////////////////////////////////
-//////////////////////////////////////////////////////////////
-
-EFilterReturn::Type FClassViewerFilterFuncs::IfInChildOfClassesSet(TSet< const UClass* >& InSet, const UClass* InClass)
-{
-	check(InClass);
-
-	if(InSet.Num())
-	{
-		// If a class is a child of any classes on this list, it will be allowed onto the list, unless it also appears on a disallowed list.
-		for( auto CurClassIt = InSet.CreateConstIterator(); CurClassIt; ++CurClassIt )
-		{
-			if(InClass->IsChildOf(*CurClassIt))
-			{
-				return EFilterReturn::Passed;
-			}
-		}
-
-		return EFilterReturn::Failed;
-	}
-
-	// Since there are none on this list, return that there is no items.
-	return EFilterReturn::NoItems;
-}
-
-EFilterReturn::Type FClassViewerFilterFuncs::IfInChildOfClassesSet(TSet< const UClass* >& InSet, const TSharedPtr< const IUnloadedBlueprintData > InClass)
-{
-	check(InClass.IsValid());
-
-	if(InSet.Num())
-	{
-		// If a class is a child of any classes on this list, it will be allowed onto the list, unless it also appears on a disallowed list.
-		for( auto CurClassIt = InSet.CreateConstIterator(); CurClassIt; ++CurClassIt )
-		{
-			if(InClass->IsChildOf(*CurClassIt))
-			{
-				return EFilterReturn::Passed;
-			}
-		}
-
-		return EFilterReturn::Failed;
-	}
-
-	// Since there are none on this list, return that there is no items.
-	return EFilterReturn::NoItems;
-}
-
-EFilterReturn::Type FClassViewerFilterFuncs::IfMatchesAllInChildOfClassesSet(TSet< const UClass* >& InSet, const UClass* InClass)
-{
-	check(InClass);
-
-	if(InSet.Num())
-	{
-		// If a class is a child of any classes on this list, it will be allowed onto the list, unless it also appears on a disallowed list.
-		for( auto CurClassIt = InSet.CreateConstIterator(); CurClassIt; ++CurClassIt )
-		{
-			if(!InClass->IsChildOf(*CurClassIt))
-			{
-				// Since it doesn't match one, it fails.
-				return EFilterReturn::Failed;
-			}
-		}
-
-		// It matches all of them, so it passes.
-		return EFilterReturn::Passed;
-	}
-
-	// Since there are none on this list, return that there is no items.
-	return EFilterReturn::NoItems;
-}
-
-EFilterReturn::Type FClassViewerFilterFuncs::IfMatchesAllInChildOfClassesSet(TSet< const UClass* >& InSet, const TSharedPtr< const IUnloadedBlueprintData > InClass)
-{
-	check(InClass.IsValid());
-
-	if(InSet.Num())
-	{
-		// If a class is a child of any classes on this list, it will be allowed onto the list, unless it also appears on a disallowed list.
-		for( auto CurClassIt = InSet.CreateConstIterator(); CurClassIt; ++CurClassIt )
-		{
-			if(!InClass->IsChildOf(*CurClassIt))
-			{
-				// Since it doesn't match one, it fails.
-				return EFilterReturn::Failed;
-			}
-		}
-
-		// It matches all of them, so it passes.
-		return EFilterReturn::Passed;
-	}
-
-	// Since there are none on this list, return that there is no items.
-	return EFilterReturn::NoItems;
-}
-
-EFilterReturn::Type FClassViewerFilterFuncs::IfMatchesAll_ObjectsSetIsAClass(TSet< const UObject* >& InSet, const UClass* InClass)
-{
-	check(InClass);
-
-	if(InSet.Num())
-	{
-		// If a class is a child of any classes on this list, it will be allowed onto the list, unless it also appears on a disallowed list.
-		for( auto CurClassIt = InSet.CreateConstIterator(); CurClassIt; ++CurClassIt )
-		{
-			if(!(*CurClassIt)->IsA(InClass))
-			{
-				// Since it doesn't match one, it fails.
-				return EFilterReturn::Failed;
-			}
-		}
-
-		// It matches all of them, so it passes.
-		return EFilterReturn::Passed;
-	}
-
-	// Since there are none on this list, return that there is no items.
-	return EFilterReturn::NoItems;
-}
-
-EFilterReturn::Type FClassViewerFilterFuncs::IfMatchesAll_ObjectsSetIsAClass(TSet< const UObject* >& InSet, const TSharedPtr< const IUnloadedBlueprintData > InClass)
-{
-	check(InClass.IsValid());
-
-	if(InSet.Num())
-	{
-		// If a class is a child of any classes on this list, it will be allowed onto the list, unless it also appears on a disallowed list.
-		for( auto CurClassIt = InSet.CreateConstIterator(); CurClassIt; ++CurClassIt )
-		{
-			if(!(*CurClassIt)->IsA(UBlueprintGeneratedClass::StaticClass()))
-			{
-				// Since it doesn't match one, it fails.
-				return EFilterReturn::Failed;
-			}
-		}
-
-		// It matches all of them, so it passes.
-		return EFilterReturn::Passed;
-	}
-
-	// Since there are none on this list, return that there is no items.
-	return EFilterReturn::NoItems;
-}
-
-EFilterReturn::Type FClassViewerFilterFuncs::IfMatchesAll_ClassesSetIsAClass(TSet< const UClass* >& InSet, const UClass* InClass)
-{
-	check(InClass);
-
-	if(InSet.Num())
-	{
-		// If a class is a child of any classes on this list, it will be allowed onto the list, unless it also appears on a disallowed list.
-		for( auto CurClassIt = InSet.CreateConstIterator(); CurClassIt; ++CurClassIt )
-		{
-			const UObject* Object = *CurClassIt;
-			if(!Object->IsA(InClass))
-			{
-				// Since it doesn't match one, it fails.
-				return EFilterReturn::Failed;
-			}
-		}
-
-		// It matches all of them, so it passes.
-		return EFilterReturn::Passed;
-	}
-
-	// Since there are none on this list, return that there is no items.
-	return EFilterReturn::NoItems;
-}
-
-EFilterReturn::Type FClassViewerFilterFuncs::IfMatchesAll_ClassesSetIsAClass(TSet< const UClass* >& InSet, const TSharedPtr< const IUnloadedBlueprintData > InClass)
-{
-	check(InClass.IsValid());
-
-	if(InSet.Num())
-	{
-		// If a class is a child of any classes on this list, it will be allowed onto the list, unless it also appears on a disallowed list.
-		for( auto CurClassIt = InSet.CreateConstIterator(); CurClassIt; ++CurClassIt )
-		{
-			const UObject* Object = *CurClassIt;
-			if(!Object->IsA(UBlueprintGeneratedClass::StaticClass()))
-			{
-				// Since it doesn't match one, it fails.
-				return EFilterReturn::Failed;
-			}
-		}
-
-		// It matches all of them, so it passes.
-		return EFilterReturn::Passed;
-	}
-
-	// Since there are none on this list, return that there is no items.
-	return EFilterReturn::NoItems;
-}
-
-EFilterReturn::Type FClassViewerFilterFuncs::IfMatches_ClassesSetIsAClass(TSet< const UClass* >& InSet, const UClass* InClass)
-{
-	check(InClass);
-
-	if(InSet.Num())
-	{
-		// If a class is a child of any classes on this list, it will be allowed onto the list, unless it also appears on a disallowed list.
-		for( auto CurClassIt = InSet.CreateConstIterator(); CurClassIt; ++CurClassIt )
-		{
-			const UObject* Object = *CurClassIt;
-			if(Object->IsA(InClass))
-			{
-				return EFilterReturn::Passed;
-			}
-		}
-
-		return EFilterReturn::Failed;
-	}
-
-	// Since there are none on this list, return that there is no items.
-	return EFilterReturn::NoItems;
-}
-
-EFilterReturn::Type FClassViewerFilterFuncs::IfMatches_ClassesSetIsAClass(TSet< const UClass* >& InSet, const TSharedPtr< const IUnloadedBlueprintData > InClass)
-{
-	check(InClass.IsValid());
-
-	if(InSet.Num())
-	{
-		// If a class is a child of any classes on this list, it will be allowed onto the list, unless it also appears on a disallowed list.
-		for( auto CurClassIt = InSet.CreateConstIterator(); CurClassIt; ++CurClassIt )
-		{
-			const UObject* Object = *CurClassIt;
-			if(Object->IsA(UBlueprintGeneratedClass::StaticClass()))
-			{
-				return EFilterReturn::Passed;
-			}
-		}
-
-		return EFilterReturn::Failed;
-	}
-
-	// Since there are none on this list, return that there is no items.
-	return EFilterReturn::NoItems;
-}
-
-EFilterReturn::Type FClassViewerFilterFuncs::IfInClassesSet(TSet< const UClass* >& InSet, const UClass* InClass)
-{
-	check(InClass);
-
-	if(InSet.Num())
-	{
-		for( auto CurClassIt = InSet.CreateConstIterator(); CurClassIt; ++CurClassIt )
-		{
-			if(InClass == *CurClassIt)
-			{
-				return EFilterReturn::Passed;
-			}
-		}
-		return EFilterReturn::Failed;
-	}
-
-	// Since there are none on this list, return that there is no items.
-	return EFilterReturn::NoItems;
-}
-
-EFilterReturn::Type FClassViewerFilterFuncs::IfInClassesSet(TSet< const UClass* >& InSet, const TSharedPtr< const IUnloadedBlueprintData > InClass)
-{
-	check(InClass.IsValid());
-
-	if(InSet.Num())
-	{
-		for( auto CurClassIt = InSet.CreateConstIterator(); CurClassIt; ++CurClassIt )
-		{
-			const TSharedPtr<const FUnloadedBlueprintData> UnloadedBlueprintData = StaticCastSharedPtr<const FUnloadedBlueprintData>(InClass);
-			if(*UnloadedBlueprintData->GetClassViewerNode().Pin()->GetClassName() == (*CurClassIt)->GetName())
-			{
-				return EFilterReturn::Passed;
-			}
-		}
-		return EFilterReturn::Failed;
-	}
-
-	// Since there are none on this list, return that there is no items.
-	return EFilterReturn::NoItems;
-}
-
-//////////////////////////////////////////////////////////////
 //////////////////////////////////////////////////////////////
 
 // A hack to make IMPLEMENT_COMPARE_CONSTREF with a templated type
@@ -500,7 +219,6 @@ namespace ClassViewer
 		static bool bPopulateClassHierarchy;
 
 		// Pre-declare these functions.
-		static bool CheckIfBlueprintBase( TSharedPtr< FClassViewerNode> InNode );
 		static UBlueprint* GetBlueprint( UClass* InClass );
 		static void UpdateClassInNode(FName InGeneratedClassPath, UClass* InNewClass, UBlueprint* InNewBluePrint );
 
@@ -526,24 +244,6 @@ namespace ClassViewer
 			return bCanCreateBlueprintOfClass;
 		}
 
-		/** Checks if a particular class is a brush.
-		 *	@param InClass				The Class to check.
-		 *	@return Returns true if the class is a brush.
-		 */
-		static bool IsBrush(const UClass* InClass)
-		{
-			return InClass->IsChildOf( ABrush::StaticClass() );
-		}
-
-		/** Checks if a particular class is placeable.
-		 *	@param InClass				The Class to check.
-		 *	@return Returns true if the class is placeable.
-		 */
-		static bool IsPlaceable(const UClass* InClass)
-		{
-			return !InClass->HasAnyClassFlags(CLASS_Abstract | CLASS_NotPlaceable) && InClass->IsChildOf<AActor>();
-		}
-
 		/** Checks if a particular class is abstract.
 		 *	@param InClass				The Class to check.
 		 *	@return Returns true if the class is abstract.
@@ -553,67 +253,6 @@ namespace ClassViewer
 			return InClass->HasAnyClassFlags(CLASS_Abstract);
 		}
 
-		/** Checks if the class is allowed under the init options of the class viewer currently building it's tree/list.
-		 *	@param InInitOptions		The class viewer's options, holds the AllowedClasses and DisallowedClasses.
-		 *	@param InClass				The class to test against.
-		 */
-		static bool IsClassAllowed( const FClassViewerInitializationOptions& InInitOptions, const TWeakObjectPtr<const UClass> InClass )
-		{
-			if(InInitOptions.ClassFilter.IsValid())
-			{
-				return InInitOptions.ClassFilter->IsClassAllowed(InInitOptions, InClass.Get(), MakeShared<FClassViewerFilterFuncs>());
-			}
-
-			return true;
-		}
-
-		/** Checks if the unloaded class is allowed under the init options of the class viewer currently building it's tree/list.
-		 *	@param InInitOptions		The class viewer's options, holds the AllowedClasses and DisallowedClasses.
-		 *	@param InNode				The node to pull the unloaded class from.
-		 */
-		bool IsClassAllowed_UnloadedBlueprint(const FClassViewerInitializationOptions& InInitOptions, TSharedPtr< FClassViewerNode > InNode )
-		{
-			if(InInitOptions.ClassFilter.IsValid() && InNode->UnloadedBlueprintData.IsValid())
-			{
-				return InInitOptions.ClassFilter->IsUnloadedClassAllowed(InInitOptions, InNode->UnloadedBlueprintData.ToSharedRef(), MakeShared<FClassViewerFilterFuncs>());
-			}
-
-			return true;
-		}
-
-		/** Checks if the TestString passes the filter.
-		 *	@param InTestString			The string to test against the filter.
-		 *	@param InTextFilter			Compiled text filter to apply.
-		 *
-		 *	@return	true if it passes the filter.
-		 */
-		static bool PassesFilter( const FString& InTestString, const FTextFilterExpressionEvaluator& InTextFilter )
-		{
-			class FClassFilterContext : public ITextFilterExpressionContext
-			{
-			public:
-				explicit FClassFilterContext(const FString& InStr)
-					: StrPtr(&InStr)
-				{
-				}
-
-				virtual bool TestBasicStringExpression(const FTextFilterString& InValue, const ETextFilterTextComparisonMode InTextComparisonMode) const override
-				{
-					return TextFilterUtils::TestBasicStringExpression(*StrPtr, InValue, InTextComparisonMode);
-				}
-
-				virtual bool TestComplexExpression(const FName& InKey, const FTextFilterString& InValue, const ETextFilterComparisonOperation InComparisonOperation, const ETextFilterTextComparisonMode InTextComparisonMode) const override
-				{
-					return false;
-				}
-
-			private:
-				const FString* StrPtr;
-			};
-
-			return InTextFilter.TestTextFilter(FClassFilterContext(InTestString));
-		}
-
 		/** Will create the instance of FClassHierarchy and populate the class hierarchy tree. */
 		static void ConstructClassHierarchy()
 		{
@@ -660,319 +299,131 @@ namespace ClassViewer
 		}
 
 		/** Recursive function to build a tree, filtering out nodes based on the InitOptions and filter search terms.
-		 *	@param InInitOptions						The class viewer's options, holds the AllowedClasses and DisallowedClasses.
 		 *	@param InOutRootNode						The node that this function will add the children of to the tree.
 		 *	@param InRootClassIndex						The index of the root node.
-		 *	@param InTextFilter							Compiled text filter to apply.
-		 *	@param bInOnlyActors						Filter option to remove non-actor classes.
-		 *	@param bInOnlyPlaceables					Filter option to remove non-placeable classes.
-		 *	@param bInOnlyBlueprintBases				Filter option to remove non-blueprint base classes.
-		 *	@param bInShowUnloadedBlueprints			Filter option to not remove unloaded blueprints due to class filter options.
-		 *  @param InAllowedDeveloperType               Filter option for dealing with developer folders.
-		 *  @param bInInternalClasses                   Filter option for showing internal classes.
-		 *  @param InternalClasses                      The classes that have been marked as Internal Only.
-		 *  @param InternalPaths                        The paths that have been marked Internal Only.
+		 *	@param InClassFilter						The class filter to use to filter nodes.
+		 *	@param InInitOptions						The class viewer's options, holds the AllowedClasses and DisallowedClasses.
 		 *
 		 *	@return Returns true if the child passed the filter.
 		 */		
-		static bool AddChildren_Tree(const FClassViewerInitializationOptions& InInitOptions, TSharedPtr< FClassViewerNode >& InOutRootNode, 
-			const TSharedPtr< FClassViewerNode >& InOriginalRootNode, const FTextFilterExpressionEvaluator& InTextFilter, bool bInOnlyActors, 
-			bool bInOnlyPlaceables, bool bInOnlyBlueprintBases, bool bInShowUnloadedBlueprints, EClassViewerDeveloperType InAllowedDeveloperType, bool bInInternalClasses,
-			const TArray<UClass*>& InternalClasses, const TArray<FDirectoryPath>& InternalPaths)
+		static bool AddChildren_Tree(TSharedPtr< FClassViewerNode >& InOutRootNode, const TSharedPtr< FClassViewerNode >& InOriginalRootNode, 
+			const TSharedPtr<FClassViewerFilter>& InClassFilter, const FClassViewerInitializationOptions& InInitOptions)
 		{
-			if (bInOnlyActors && *InOriginalRootNode->GetClassName().Get() != FString(TEXT("Actor")))
+			InOutRootNode->bPassesFilter = InClassFilter->IsNodeAllowed(InInitOptions, InOutRootNode.ToSharedRef());
+			if (!InOutRootNode->bPassesFilter)
 			{
-				InOutRootNode->bPassesFilter = false;
 				return false;
 			}
 
-			bool bChildrenPassesFilter(false);
-			bool bReturnPassesFilter(false);
-
-			bool bPassesBlueprintBaseFilter = !bInOnlyBlueprintBases || CheckIfBlueprintBase(InOriginalRootNode);
-			bool bIsUnloadedBlueprint = !InOriginalRootNode->Class.IsValid();
-			bool bPassesPlaceableFilter = false;
-
-			// Determine if we allow any developer folder classes, if so determine if this class is in one of the allowed developer folders.
-			bool bPassesDeveloperFilter = true;
-			static const FString DeveloperPathWithSlash = FPackageName::FilenameToLongPackageName(FPaths::GameDevelopersDir());
-			static const FString UserDeveloperPathWithSlash = FPackageName::FilenameToLongPackageName(FPaths::GameUserDeveloperDir());
-			FString GeneratedClassPathString = InOriginalRootNode->ClassPath.ToString();
-
-			if (InAllowedDeveloperType == EClassViewerDeveloperType::CVDT_None)
-			{
-				bPassesDeveloperFilter = !GeneratedClassPathString.StartsWith(DeveloperPathWithSlash);
-			}
-			else if (InAllowedDeveloperType == EClassViewerDeveloperType::CVDT_CurrentUser)
-			{
-				if (GeneratedClassPathString.StartsWith(DeveloperPathWithSlash))
-					bPassesDeveloperFilter = GeneratedClassPathString.StartsWith(UserDeveloperPathWithSlash);
-			}
-
-			// The INI files declare classes and folders that are considered internal only. Does this class match any of those patterns?
-			// INI path: /Script/ClassViewer.ClassViewerProjectSettings
-			bool bPassesInternalFilter = true;
-			if (!bInInternalClasses && InternalPaths.Num() > 0)
-			{
-				for (int i = 0; i < InternalPaths.Num(); i++)
-				{
-					if (GeneratedClassPathString.StartsWith(InternalPaths[i].Path))
-					{
-						bPassesInternalFilter = false;
-						break;
-					}
-				}
-			}
-			if (!bInInternalClasses && InternalClasses.Num() > 0 && bPassesInternalFilter && InOriginalRootNode->Class.IsValid())
-			{
-				for (int i = 0; i < InternalClasses.Num(); i++)
-				{
-					if (InOriginalRootNode->Class->IsChildOf(InternalClasses[i]))
-					{
-						bPassesInternalFilter = false;
-						break;
-					}
-				}
-			}
-
-			// When in picker mode, brushes are valid "placeable" actors
-			if (bInOnlyPlaceables && InInitOptions.Mode == EClassViewerMode::ClassPicker && IsBrush(InOriginalRootNode->Class.Get()) && IsPlaceable(InOriginalRootNode->Class.Get()))
-			{
-				bPassesPlaceableFilter = true;
-			}
-			else
-			{
-				bPassesPlaceableFilter = !bInOnlyPlaceables || InOriginalRootNode->IsClassPlaceable();
-			}
-
-			// There are few options for filtering an unloaded blueprint, if it matches with this filter, it passes.
-			if(bIsUnloadedBlueprint)
-			{
-				if(bInShowUnloadedBlueprints)
-				{
-					bReturnPassesFilter = InOutRootNode->bPassesFilter = bPassesPlaceableFilter && bPassesBlueprintBaseFilter && bPassesDeveloperFilter && bPassesInternalFilter && IsClassAllowed_UnloadedBlueprint(InInitOptions, InOriginalRootNode) && PassesFilter(*InOriginalRootNode->GetClassName().Get(), InTextFilter);
-				}
-			}
-			else
-			{
-				bReturnPassesFilter = InOutRootNode->bPassesFilter = bPassesPlaceableFilter && bPassesBlueprintBaseFilter && bPassesDeveloperFilter && bPassesInternalFilter && IsClassAllowed(InInitOptions, InOriginalRootNode->Class) && PassesFilter(*InOriginalRootNode->GetClassName().Get(), InTextFilter);
-			}
-
+			bool bReturnPassesFilter = false;
 
 			TArray< TSharedPtr< FClassViewerNode > >& ChildList = InOriginalRootNode->GetChildrenList();
 			for(int32 ChildIdx = 0; ChildIdx < ChildList.Num(); ChildIdx++)
 			{
 				TSharedPtr< FClassViewerNode > NewNode = MakeShareable( new FClassViewerNode( *ChildList[ChildIdx].Get() ) );
 
-				bReturnPassesFilter |= bChildrenPassesFilter = AddChildren_Tree(InInitOptions, NewNode, ChildList[ChildIdx], InTextFilter, false, /* bInOnlyActors - false so that anything below Actor is added */
-					bInOnlyPlaceables, bInOnlyBlueprintBases, bInShowUnloadedBlueprints, InAllowedDeveloperType, bInInternalClasses, InternalClasses, InternalPaths);
-				if(bChildrenPassesFilter)
+				bool bChildrenPassesFilter = AddChildren_Tree(NewNode, ChildList[ChildIdx], InClassFilter, InInitOptions);
+				bReturnPassesFilter |= bChildrenPassesFilter;
+
+				if (bChildrenPassesFilter)
 				{
 					InOutRootNode->AddChild(NewNode);
 				}
-
 			}
 
 			return bReturnPassesFilter;
 		}
 
 		/** Builds the class tree.
-		 *	@param InInitOptions						The class viewer's options, holds the AllowedClasses and DisallowedClasses.
 		 *	@param InOutRootNode						The node to root the tree to.
-		 *	@param InTextFilter							Compiled text filter to apply.
-		 *	@param bInOnlyPlaceables					Filter option to remove non-placeable classes.
-		 *	@param bInOnlyActors						Filter option to root the tree in the "Actor" class.
-		 *	@param bInOnlyBlueprintBases				Filter option to remove non-blueprint base classes.
-		 *	@param bInShowUnloadedBlueprints			Filter option to not remove unloaded blueprints due to class filter options.
-		 *  @param InAllowedDeveloperType               Filter option for dealing with developer folders.
-		 *  @param bInInternalClasses                   Filter option for showing internal classes.
-		 *  @param InternalClasses                      The classes that have been marked as Internal Only.
-		 *  @param InternalPaths                        The paths that have been marked Internal Only.
+		 *	@param InClassFilter						The class filter to use to filter nodes.
+		 *	@param InInitOptions						The class viewer's options, holds the AllowedClasses and DisallowedClasses.
+		 *
 		 *	@return A fully built tree.
 		 */
-		static void GetClassTree(const FClassViewerInitializationOptions& InInitOptions, TSharedPtr< FClassViewerNode >& InOutRootNode, 
-			const FTextFilterExpressionEvaluator& InTextFilter, bool bInOnlyPlaceables, bool bInOnlyActors, bool bInOnlyBlueprintBases, 
-			bool bInShowUnloadedBlueprints,	EClassViewerDeveloperType InAllowedDeveloperType = EClassViewerDeveloperType::CVDT_All, bool bInInternalClasses = true, 
-			const TArray<UClass*>& InternalClasses = TArray<UClass*>(), const TArray<FDirectoryPath>& InternalPaths = TArray<FDirectoryPath>())
+		static void GetClassTree(TSharedPtr< FClassViewerNode >& InOutRootNode, const TSharedPtr<FClassViewerFilter>& InClassFilter,
+			const FClassViewerInitializationOptions& InInitOptions)
 		{
-			
 			const TSharedPtr< FClassViewerNode > ObjectClassRoot = ClassHierarchy->GetObjectRootNode();
 
 			// Duplicate the node, it will have no children.
 			InOutRootNode = MakeShareable(new FClassViewerNode(*ObjectClassRoot));
 
-			if(bInOnlyActors)
+			if (InInitOptions.bIsActorsOnly)
 			{
-				for(int32 ClassIdx = 0; ClassIdx < ObjectClassRoot->GetChildrenList().Num(); ClassIdx++)
+				for (int32 ClassIdx = 0; ClassIdx < ObjectClassRoot->GetChildrenList().Num(); ClassIdx++)
 				{
 					TSharedPtr<FClassViewerNode> ChildNode = MakeShareable(new FClassViewerNode(*ObjectClassRoot->GetChildrenList()[ClassIdx].Get()));
-					if (AddChildren_Tree(InInitOptions, ChildNode, ObjectClassRoot->GetChildrenList()[ClassIdx], InTextFilter, true, bInOnlyPlaceables, bInOnlyBlueprintBases, bInShowUnloadedBlueprints, InAllowedDeveloperType, bInInternalClasses, InternalClasses, InternalPaths))
+					if (AddChildren_Tree(ChildNode, ObjectClassRoot->GetChildrenList()[ClassIdx], InClassFilter, InInitOptions))
+					{
 						InOutRootNode->AddChild(ChildNode);
+					}
 				}
 			}
 			else
 			{
-				AddChildren_Tree(InInitOptions, InOutRootNode, ObjectClassRoot, InTextFilter, false, bInOnlyPlaceables, bInOnlyBlueprintBases, bInShowUnloadedBlueprints, InAllowedDeveloperType, bInInternalClasses, InternalClasses, InternalPaths);
+				AddChildren_Tree(InOutRootNode, ObjectClassRoot, InClassFilter, InInitOptions);
 			}
 		}
 
 		/** Recursive function to build the list, filtering out nodes based on the InitOptions and filter search terms.
-		 *	@param InInitOptions						The class viewer's options, holds the AllowedClasses and DisallowedClasses.
 		 *	@param InOutRootNode						The node that this function will add the children of to the tree.
 		 *	@param InRootClassIndex						The index of the root node.
-		 *	@param InTextFilter							Compiled text filter to apply.
-		 *	@param bInOnlyActors						Filter option to remove non-actor classes.
-		 *	@param bInOnlyPlaceables					Filter option to remove non-placeable classes.
-		 *	@param bInOnlyBlueprintBases				Filter option to remove non-blueprint base classes.
-		 *	@param bInShowUnloadedBlueprints			Filter option to not remove unloaded blueprints due to class filter options.
-		 *  @param InAllowedDeveloperType               Filter option for dealing with developer folders.
-		 *  @param bInInternalClasses                   Filter option for showing internal classes.
-		 *  @param InternalClasses                      The classes that have been marked as Internal Only.
-		 *  @param InternalPaths                        The paths that have been marked Internal Only.
+		 *	@param InClassFilter						The class filter to use to filter nodes.
+		 *	@param InInitOptions						The class viewer's options, holds the AllowedClasses and DisallowedClasses.
 		 *
 		 *	@return Returns true if the child passed the filter.
 		 */
-		static void AddChildren_List(const FClassViewerInitializationOptions& InInitOptions, TArray< TSharedPtr< FClassViewerNode > >& InOutNodeList,
-			const TSharedPtr< FClassViewerNode >& InOriginalRootNode, const FTextFilterExpressionEvaluator& InTextFilter,
-			bool bInOnlyActors, bool bInOnlyPlaceables, bool bInOnlyBlueprintBases, bool bInShowUnloadedBlueprints,
-			EClassViewerDeveloperType InAllowedDeveloperType, bool bInInternalClasses,
-			const TArray<UClass*>& InternalClasses, const TArray<FDirectoryPath>& InternalPaths)
+		static void AddChildren_List(TArray< TSharedPtr< FClassViewerNode > >& InOutNodeList, const TSharedPtr< FClassViewerNode >& InOriginalRootNode, 
+			const TSharedPtr< FClassViewerFilter >& InClassFilter, const FClassViewerInitializationOptions& InInitOptions)
 		{
-			if (bInOnlyActors && *InOriginalRootNode->GetClassName().Get() != FString(TEXT("Actor")))
-			{
-				return;
-			}
-
-			bool bPassesBlueprintBaseFilter = !bInOnlyBlueprintBases || CheckIfBlueprintBase(InOriginalRootNode);
-			bool bIsUnloadedBlueprint = !InOriginalRootNode->Class.IsValid();
-			bool bPassesPlaceableFilter = false;
-			const bool bPassesEditorClassTest = !InInitOptions.bEditorClassesOnly || InOriginalRootNode->IsEditorOnlyClass();
-
-			// Determine if we allow any developer folder classes, if so determine if this class is in one of the allowed developer folders.
-			bool bPassesDeveloperFilter = true;
-			static const FString DeveloperPathWithSlash = FPackageName::FilenameToLongPackageName(FPaths::GameDevelopersDir());
-			static const FString UserDeveloperPathWithSlash = FPackageName::FilenameToLongPackageName(FPaths::GameUserDeveloperDir());
-			FString GeneratedClassPathString = InOriginalRootNode->ClassPath.ToString();
-			if (InAllowedDeveloperType == EClassViewerDeveloperType::CVDT_None)
+			if (InClassFilter->IsNodeAllowed(InInitOptions, InOriginalRootNode.ToSharedRef()))
 			{
-				bPassesDeveloperFilter = !GeneratedClassPathString.StartsWith(DeveloperPathWithSlash);
-			}
-			else if (InAllowedDeveloperType == EClassViewerDeveloperType::CVDT_CurrentUser)
-			{
-				if (GeneratedClassPathString.StartsWith(DeveloperPathWithSlash))
-					bPassesDeveloperFilter = GeneratedClassPathString.StartsWith(UserDeveloperPathWithSlash);
-			}
+				TSharedPtr< FClassViewerNode > NewNode = MakeShareable(new FClassViewerNode(*InOriginalRootNode.Get()));
+				NewNode->bPassesFilter = true;
+				NewNode->PropertyHandle = InOriginalRootNode->PropertyHandle;
 
-			// The INI files declare classes and folders that are considered internal only. Does this class match any of those patterns?
-			// INI path: /Script/ClassViewer.ClassViewerProjectSettings
-			bool bPassesInternalFilter = true;
-			if (!bInInternalClasses && InternalPaths.Num() > 0)
-			{
-				for (int i = 0; i < InternalPaths.Num(); i++)
-				{
-					if (GeneratedClassPathString.StartsWith(InternalPaths[i].Path))
-					{
-						bPassesInternalFilter = false;
-						break;
-					}
-				}
-			}
-			if (!bInInternalClasses && InternalClasses.Num() > 0 && bPassesInternalFilter)
-			{
-				for (int i = 0; i < InternalClasses.Num(); i++)
-				{
-					if (InOriginalRootNode->Class->IsChildOf(InternalClasses[i]))
-					{
-						bPassesInternalFilter = false;
-						break;
-					}
-				}
-			}
-
-			// When in picker mode, brushes are valid "placeable" actors
-			if( bInOnlyPlaceables && InInitOptions.Mode == EClassViewerMode::ClassPicker && IsBrush(InOriginalRootNode->Class.Get()) && IsPlaceable(InOriginalRootNode->Class.Get()))
-			{
-				bPassesPlaceableFilter = true;
-			}
-			else
-			{
-				bPassesPlaceableFilter = !bInOnlyPlaceables || InOriginalRootNode->IsClassPlaceable();
-			}
-
-			TSharedPtr< FClassViewerNode > NewNode = MakeShareable( new FClassViewerNode( *InOriginalRootNode.Get() ) );
-			
-			// There are few options for filtering an unloaded blueprint, if it matches with this filter, it passes.
-			if(bIsUnloadedBlueprint)
-			{
-				if(bInShowUnloadedBlueprints)
-				{
-					NewNode->bPassesFilter = bPassesPlaceableFilter && bPassesBlueprintBaseFilter && bPassesDeveloperFilter && bPassesInternalFilter && bPassesEditorClassTest
-						&& IsClassAllowed_UnloadedBlueprint(InInitOptions, InOriginalRootNode) 
-						&& PassesFilter(*InOriginalRootNode->GetClassName().Get(), InTextFilter);
-				}
-			}
-			else
-			{
-				NewNode->bPassesFilter = bPassesPlaceableFilter && bPassesBlueprintBaseFilter && bPassesDeveloperFilter && bPassesInternalFilter && bPassesEditorClassTest
-					&&  IsClassAllowed(InInitOptions, InOriginalRootNode->Class) 
-					&& PassesFilter(*InOriginalRootNode->GetClassName().Get(), InTextFilter);
-			}
-
-			if(NewNode->bPassesFilter)
-			{
 				InOutNodeList.Add(NewNode);
 			}
 
-			NewNode->PropertyHandle = InInitOptions.PropertyHandle;
-
-			TArray< TSharedPtr< FClassViewerNode > >& ChildList = InOriginalRootNode->GetChildrenList();
-			for(int32 ChildIdx = 0; ChildIdx < ChildList.Num(); ChildIdx++)
+			for(const TSharedPtr<FClassViewerNode>& ChildNode : InOriginalRootNode->GetChildrenList())
 			{
-				AddChildren_List(InInitOptions, InOutNodeList, ChildList[ChildIdx], InTextFilter, false, /* bInOnlyActors - false so that anything below Actor is added */
-					bInOnlyPlaceables, bInOnlyBlueprintBases, bInShowUnloadedBlueprints, InAllowedDeveloperType, bInInternalClasses, InternalClasses, InternalPaths);
+				FClassViewerInitializationOptions TempOptions = InInitOptions;
+
+				// set bOnlyActors to false so that anything below Actor is added
+				TempOptions.bIsActorsOnly = false;
+				AddChildren_List(InOutNodeList, ChildNode, InClassFilter, InInitOptions);
 			}
 		}	
 
 		/** Builds the class list.
-		 *	@param InInitOptions						The class viewer's options, holds the AllowedClasses and DisallowedClasses.
 		 *	@param InOutNodeList						The list to add all the nodes to.
-		 *	@param InTextFilter							Compiled text filter to apply.
-		 *	@param bInOnlyPlaceables					Filter option to remove non-placeable classes.
-		 *	@param bInOnlyActors						Filter option to root the tree in the "Actor" class.
-		 *	@param bInOnlyBlueprintBases				Filter option to remove non-blueprint base classes.
-		 *	@param bInShowUnloadedBlueprints			Filter option to not remove unloaded blueprints due to class filter options.
-		 *  @param InAllowedDeveloperType               Filter option for dealing with developer folders.
-		 *  @param bInInternalClasses                   Filter option for showing internal classes.
-		 *  @param InternalClasses                      The classes that have been marked as Internal Only.
-		 *  @param InternalPaths                        The paths that have been marked Internal Only.
+		 *	@param InClassFilter						The class filter to use to filter nodes.
+		 *	@param InInitOptions						The class viewer's options, holds the AllowedClasses and DisallowedClasses.
 		 *
 		 *	@return A fully built list.
 		 */
-		static void GetClassList(const FClassViewerInitializationOptions& InInitOptions, TArray< TSharedPtr< FClassViewerNode > >& InOutNodeList, 
-			const FTextFilterExpressionEvaluator& InTextFilter, bool bInOnlyPlaceables, bool bInOnlyActors, bool bInOnlyBlueprintBases, bool bInShowUnloadedBlueprints,
-			EClassViewerDeveloperType InAllowedDeveloperType = EClassViewerDeveloperType::CVDT_All, bool bInInternalClasses = true,
-			const TArray<UClass*>& InternalClasses = TArray<UClass*>(), const TArray<FDirectoryPath>& InternalPaths = TArray<FDirectoryPath>())
+		static void GetClassList(TArray< TSharedPtr< FClassViewerNode > >& InOutNodeList, const TSharedPtr<FClassViewerFilter>& InClassFilter, 
+			const FClassViewerInitializationOptions& InInitOptions)
 		{
 			const TSharedPtr< FClassViewerNode > ObjectClassRoot = ClassHierarchy->GetObjectRootNode();
 
 			// If the option to see the object root class is set, add it to the list, proceed normally from there so the actor's only filter continues to work.
-			if(InInitOptions.bShowObjectRootClass)
+			if (InInitOptions.bShowObjectRootClass)
 			{
-				TSharedPtr< FClassViewerNode > NewNode = MakeShareable( new FClassViewerNode( *ObjectClassRoot.Get() ) );
-				NewNode->bPassesFilter = IsClassAllowed(InInitOptions, ObjectClassRoot->Class) && PassesFilter(*ObjectClassRoot->GetClassName().Get(), InTextFilter);
-				
-				if(NewNode->bPassesFilter)
+				if (InClassFilter->IsNodeAllowed(InInitOptions, ObjectClassRoot.ToSharedRef()))
 				{
+					TSharedPtr< FClassViewerNode > NewNode = MakeShareable(new FClassViewerNode(*ObjectClassRoot.Get()));
+					NewNode->bPassesFilter = true;
+					NewNode->PropertyHandle = InInitOptions.PropertyHandle;
+
 					InOutNodeList.Add(NewNode);
 				}
-
-				NewNode->PropertyHandle = InInitOptions.PropertyHandle;
 			}
 
-			TArray< TSharedPtr< FClassViewerNode > >& ChildList = ObjectClassRoot->GetChildrenList();
-			for(int32 ObjectChildIndex = 0; ObjectChildIndex < ChildList.Num(); ObjectChildIndex++)
+			for(const TSharedPtr<FClassViewerNode>& ChildNode : ObjectClassRoot->GetChildrenList())
 			{
-				AddChildren_List(InInitOptions, InOutNodeList, ChildList[ObjectChildIndex], InTextFilter, bInOnlyActors, bInOnlyPlaceables, bInOnlyBlueprintBases, bInShowUnloadedBlueprints, InAllowedDeveloperType, bInInternalClasses, InternalClasses, InternalPaths);
+				AddChildren_List(InOutNodeList, ChildNode, InClassFilter, InInitOptions);
 			}
 		}
 
@@ -1012,29 +463,6 @@ namespace ClassViewer
 			}
 		}
 
-		/** Checks if a node is a blueprint base or not.
-		 *	@param	InNode					The node to check if it is a blueprint base.
-		 *
-		 *	@return							true if the class is a blueprint.
-		 */
-		static bool CheckIfBlueprintBase( TSharedPtr< FClassViewerNode> InNode )
-		{
-			// If there is no class, it may be an unloaded blueprint.
-			if(UClass* Class = InNode->Class.Get())
-			{
-				return CanCreateBlueprintOfClass_IgnoreDeprecation(Class);
-			}
-			else if(InNode->bIsBPNormalType)
-			{
-				bool bAllowDerivedBlueprints = false;
-				GConfig->GetBool(TEXT("Kismet"), TEXT("AllowDerivedBlueprints"), /*out*/ bAllowDerivedBlueprints, GEngineIni);
-
-				return bAllowDerivedBlueprints;
-			}
-
-			return false;
-		}
-
 		/**
 		 * Creates a blueprint from a class.
 		 *
@@ -1616,21 +1044,6 @@ FClassHierarchy::~FClassHierarchy()
 	FModuleManager::Get().OnModulesChanged().RemoveAll(this);
 }
 
-static TSharedPtr< FClassViewerNode > CreateNodeForClass(UClass* Class, const TMultiMap<FName, FAssetData>& BlueprintPackageToAssetDataMap)
-{
-	// Create the new node so it can be passed to AddChildren, fill it in with if it is placeable, abstract, and/or a brush.
-	TSharedPtr< FClassViewerNode > NewNode = MakeShareable(new FClassViewerNode(Class->GetName(), Class->GetDisplayNameText().ToString()));
-	NewNode->Blueprint = ClassViewer::Helpers::GetBlueprint(Class);
-	NewNode->Class = Class;
-	NewNode->ClassPath = FName(*Class->GetPathName());
-	if (Class->GetSuperClass())
-	{
-		NewNode->ParentClassPath = FName(*Class->GetSuperClass()->GetPathName());
-	}
-
-	return NewNode;
-}
-
 void FClassHierarchy::OnHotReload( bool bWasTriggeredAutomatically )
 {
 	ClassViewer::Helpers::RequestPopulateClassHierarchy();
@@ -1674,13 +1087,13 @@ void FClassHierarchy::AddChildren_NoFilter( TSharedPtr< FClassViewerNode >& InOu
 				TSharedPtr<FClassViewerNode>& ParentEntry = Nodes.FindOrAdd(CurrentClass->GetSuperClass());
 				if ( !ParentEntry.IsValid() )
 				{
-					ParentEntry = CreateNodeForClass(CurrentClass->GetSuperClass(), BlueprintPackageToAssetDataMap);
+					ParentEntry = MakeShareable(new FClassViewerNode(CurrentClass->GetSuperClass()));
 				}
 
 				TSharedPtr<FClassViewerNode>& MyEntry = Nodes.FindOrAdd(CurrentClass);
 				if ( !MyEntry.IsValid() )
 				{
-					MyEntry = CreateNodeForClass(CurrentClass, BlueprintPackageToAssetDataMap);
+					MyEntry = MakeShareable(new FClassViewerNode(CurrentClass));
 				}
 
 				if ( !Visited.Contains(CurrentClass) )
@@ -1929,12 +1342,13 @@ void FClassHierarchy::LoadUnloadedTagData(TSharedPtr<FClassViewerNode>& InOutCla
 		InOutClassViewerNode->ParentClassPath = FName(*FPackageName::ExportTextPathToObjectPath(ParentClassPathString));
 	}
 
-	InOutClassViewerNode->bIsBPNormalType = InAssetData.GetTagValueRef<FString>(FBlueprintTags::BlueprintType) == TEXT("BPType_Normal");
-			
 	// It is an unloaded blueprint, so we need to create the structure that will hold the data.
 	TSharedPtr<FUnloadedBlueprintData> UnloadedBlueprintData = MakeShareable( new FUnloadedBlueprintData(InOutClassViewerNode) );
 	InOutClassViewerNode->UnloadedBlueprintData = UnloadedBlueprintData;
 
+	const bool bNormalBlueprintType = InAssetData.GetTagValueRef<FString>(FBlueprintTags::BlueprintType) == TEXT("BPType_Normal");
+	InOutClassViewerNode->UnloadedBlueprintData->SetNormalBlueprintType( bNormalBlueprintType );
+
 	// Get the class flags.
 	const uint32 ClassFlags = InAssetData.GetTagValueRef<uint32>(FBlueprintTags::ClassFlags);
 	InOutClassViewerNode->UnloadedBlueprintData->SetClassFlags(ClassFlags);
@@ -2039,8 +1453,6 @@ void SClassViewer::Construct(const FArguments& InArgs, const FClassViewerInitial
 	bNeedsRefresh = true;
 	NumClasses = 0;
 	
-	bCanShowInternalClasses = true;
-	
 	// Listen for when view settings are changed
 	UClassViewerSettings::OnSettingChanged().AddSP(this, &SClassViewer::HandleSettingChanged);
 
@@ -2048,22 +1460,15 @@ void SClassViewer::Construct(const FArguments& InArgs, const FClassViewerInitial
 
 	OnClassPicked = InArgs._OnClassPickedDelegate;
 
-	TextFilterPtr = MakeShareable(new FTextFilterExpressionEvaluator(ETextFilterExpressionEvaluatorMode::BasicString));
-
 	bSaveExpansionStates = true;
 	bPendingSetExpansionStates = false;
 
+	ClassFilter = MakeShareable(new FClassViewerFilter(InitOptions));
+
 	bEnableClassDynamicLoading = InInitOptions.bEnableClassDynamicLoading;
 
 	EVisibility HeaderVisibility = (this->InitOptions.Mode == EClassViewerMode::ClassBrowsing)? EVisibility::Visible : EVisibility::Collapsed;
 
-	// Set these values to the user specified settings.
-	bIsActorsOnly = InitOptions.bIsActorsOnly | InitOptions.bIsPlaceableOnly;
-	bIsPlaceableOnly = InitOptions.bIsPlaceableOnly;
-	bIsBlueprintBaseOnly = InitOptions.bIsBlueprintBaseOnly;
-	bShowUnloadedBlueprints = InitOptions.bShowUnloadedBlueprints;
-	bool bHasTitle = InitOptions.ViewerTitleString.IsEmpty() == false;
-
 	// If set to default, decide what display mode to use.
 	if( InitOptions.DisplayMode == EClassViewerDisplayMode::DefaultView )
 	{
@@ -2167,6 +1572,8 @@ void SClassViewer::Construct(const FArguments& InArgs, const FClassViewerInitial
 		);
 	TSharedRef<STreeView<TSharedPtr< FClassViewerNode > > > ClassTreeView = ClassTree.ToSharedRef();
 	TSharedRef<SListView<TSharedPtr< FClassViewerNode > > > ClassListView = ClassList.ToSharedRef();
+	
+	bool bHasTitle = InitOptions.ViewerTitleString.IsEmpty() == false;
 
 	// Holds the bulk of the class viewer's sub-widgets, to be added to the widget after construction
 	TSharedPtr< SWidget > ClassViewerContent;
@@ -2616,7 +2023,6 @@ EClassViewerDeveloperType SClassViewer::GetCurrentDeveloperViewType() const
 	{
 		return EClassViewerDeveloperType::CVDT_All;
 	}
-
 	return GetDefault<UClassViewerSettings>()->DeveloperFolderType;
 }
 
@@ -2631,7 +2037,6 @@ void SClassViewer::GetInternalOnlyClasses(TArray<FSoftClassPath>& Classes)
 	{
 		return;
 	}
-
 	Classes = GetDefault<UClassViewerProjectSettings>()->InternalOnlyClasses;
 }
 
@@ -2644,7 +2049,6 @@ void SClassViewer::GetInternalOnlyPaths(TArray<FDirectoryPath>& Paths)
 	Paths = GetDefault<UClassViewerProjectSettings>()->InternalOnlyPaths;
 }
 
-
 FText SClassViewer::GetClassCountText() const
 {
 
@@ -2743,8 +2147,8 @@ void SClassViewer::FindInContentBrowser()
 void SClassViewer::OnFilterTextChanged( const FText& InFilterText )
 {
 	// Update the compiled filter and report any syntax error information back to the user
-	TextFilterPtr->SetFilterText(InFilterText);
-	SearchBox->SetError(TextFilterPtr->GetFilterErrorText());
+	ClassFilter->TextFilter->SetFilterText(InFilterText);
+	SearchBox->SetError(ClassFilter->TextFilter->GetFilterErrorText());
 
 	// Repopulate the list to show only what has not been filtered out.
 	Refresh();
@@ -2790,12 +2194,12 @@ bool SClassViewer::Menu_CanExecute() const
 
 void SClassViewer::MenuActorsOnly_Execute()
 {
-	bIsActorsOnly = !bIsActorsOnly;
+	InitOptions.bIsActorsOnly = !InitOptions.bIsActorsOnly;
 
 	// "Placeable Only" cannot be true when "Actors Only" is false.
-	if(!bIsActorsOnly)
+	if (!InitOptions.bIsActorsOnly)
 	{
-		bIsPlaceableOnly = false;
+		InitOptions.bIsPlaceableOnly = false;
 	}
 
 	Refresh();
@@ -2803,17 +2207,17 @@ void SClassViewer::MenuActorsOnly_Execute()
 
 bool SClassViewer::MenuActorsOnly_IsChecked() const
 {
-	return bIsActorsOnly;
+	return InitOptions.bIsActorsOnly;
 }
 
 void SClassViewer::MenuPlaceableOnly_Execute()
 {
-	bIsPlaceableOnly = !bIsPlaceableOnly;
+	InitOptions.bIsPlaceableOnly = !InitOptions.bIsPlaceableOnly;
 
 	// "Actors Only" must be true when "Placeable Only" is true.
-	if(bIsPlaceableOnly)
+	if (InitOptions.bIsPlaceableOnly)
 	{
-		bIsActorsOnly = true;
+		InitOptions.bIsActorsOnly = true;
 	}
 
 	Refresh();
@@ -2821,19 +2225,19 @@ void SClassViewer::MenuPlaceableOnly_Execute()
 
 bool SClassViewer::MenuPlaceableOnly_IsChecked() const
 {
-	return bIsPlaceableOnly;
+	return InitOptions.bIsPlaceableOnly;
 }
 
 void SClassViewer::MenuBlueprintBasesOnly_Execute()
 {
-	bIsBlueprintBaseOnly = !bIsBlueprintBaseOnly;
+	InitOptions.bIsBlueprintBaseOnly = !InitOptions.bIsBlueprintBaseOnly;
 
 	Refresh();
 }
 
 bool SClassViewer::MenuBlueprintBasesOnly_IsChecked() const
 {
-	return bIsBlueprintBaseOnly;
+	return InitOptions.bIsBlueprintBaseOnly;
 }
 
 TSharedRef<SWidget> SClassViewer::FillFilterEntries()
@@ -2975,15 +2379,11 @@ void SClassViewer::Populate()
 	// Empty the tree out so it can be redone.
 	RootTreeItems.Empty();
 
-	bool ShowingInternalClasses = IsShowingInternalClasses();
-
 	TArray<FSoftClassPath> InternalClassNames;
-	TArray<UClass*> InternalClasses;
-	TArray<FDirectoryPath> InternalPaths;
 	// If we aren't showing the internal classes, then we need to know what classes to consider Internal Only, so let's gather them up from the settings object.
-	if (!ShowingInternalClasses)
+	if (!IsShowingInternalClasses())
 	{
-		GetInternalOnlyPaths(InternalPaths);
+		GetInternalOnlyPaths(ClassFilter->InternalPaths);
 		GetInternalOnlyClasses(InternalClassNames);
 
 		// Take the package names for the internal only classes and convert them into their UClass
@@ -2994,11 +2394,13 @@ void SClassViewer::Populate()
 
 			if (ClassNode.IsValid())
 			{
-				InternalClasses.Add(ClassNode->Class.Get());
+				ClassFilter->InternalClasses.Add(ClassNode->Class.Get());
 			}
 		}
 	}
 
+	
+
 	// Based on if the list or tree is visible we create what will be displayed differently.
 	if(InitOptions.DisplayMode == EClassViewerDisplayMode::TreeView)
 	{
@@ -3006,11 +2408,10 @@ void SClassViewer::Populate()
 		TSharedPtr<FClassViewerNode> RootNode;
 
 		// Get the class tree, passing in certain filter options.
-		ClassViewer::Helpers::GetClassTree(InitOptions, RootNode, *TextFilterPtr, MenuPlaceableOnly_IsChecked(), MenuActorsOnly_IsChecked(), MenuBlueprintBasesOnly_IsChecked(),
-			bShowUnloadedBlueprints, GetCurrentDeveloperViewType(), ShowingInternalClasses, InternalClasses, InternalPaths);
+		ClassViewer::Helpers::GetClassTree(RootNode, ClassFilter, InitOptions);
 
 		// Check if we will restore expansion states, we will not if there is filtering happening.
-		const bool bRestoreExpansionState = TextFilterPtr->GetFilterType() == ETextFilterExpressionType::Empty;
+		const bool bRestoreExpansionState = ClassFilter->TextFilter->GetFilterType() == ETextFilterExpressionType::Empty;
 
 		if(InitOptions.bShowObjectRootClass)
 		{
@@ -3022,7 +2423,7 @@ void SClassViewer::Populate()
 			}
 
 			// Expand any items that pass the filter.
-			if(TextFilterPtr->GetFilterType() != ETextFilterExpressionType::Empty)
+			if(ClassFilter->TextFilter->GetFilterType() != ETextFilterExpressionType::Empty)
 			{
 				ExpandFilteredInNodes(RootNode);
 			}
@@ -3039,7 +2440,7 @@ void SClassViewer::Populate()
 				}
 
 				// Expand any items that pass the filter.
-				if(TextFilterPtr->GetFilterType() != ETextFilterExpressionType::Empty)
+				if(ClassFilter->TextFilter->GetFilterType() != ETextFilterExpressionType::Empty)
 				{
 					ExpandFilteredInNodes(RootNode->GetChildrenList()[ChildIndex]);
 				}
@@ -3064,8 +2465,7 @@ void SClassViewer::Populate()
 	else
 	{
 		// Get the class list, passing in certain filter options.
-		ClassViewer::Helpers::GetClassList(InitOptions, RootTreeItems, *TextFilterPtr, MenuPlaceableOnly_IsChecked(), MenuActorsOnly_IsChecked(), MenuBlueprintBasesOnly_IsChecked(), 
-			bShowUnloadedBlueprints, GetCurrentDeveloperViewType(), ShowingInternalClasses, InternalClasses, InternalPaths);
+		ClassViewer::Helpers::GetClassList(RootTreeItems, ClassFilter, InitOptions);
 
 		// Sort the list alphabetically.
 		struct FCompareFClassViewerNode
@@ -3180,7 +2580,7 @@ void SClassViewer::Tick( const FGeometry& AllottedGeometry, const double InCurre
 
 bool SClassViewer::IsClassAllowed(const UClass* InClass) const
 {
-	return ClassViewer::Helpers::IsClassAllowed(InitOptions, InClass);
+	return ClassFilter->IsClassAllowed(InitOptions, InClass, ClassFilter->FilterFunctions);
 }
 
 void SClassViewer::HandleSettingChanged(FName PropertyName)
@@ -3202,7 +2602,7 @@ void SClassViewer::ToggleShowInternalClasses()
 
 bool SClassViewer::IsToggleShowInternalClassesAllowed() const
 {
-	return bCanShowInternalClasses;
+	return InitOptions.bAllowViewOptions;
 }
 
 bool SClassViewer::IsShowingInternalClasses() const
diff --git a/Engine/Source/Editor/ClassViewer/Private/SClassViewer.h b/Engine/Source/Editor/ClassViewer/Private/SClassViewer.h
--- a/Engine/Source/Editor/ClassViewer/Private/SClassViewer.h
+++ b/Engine/Source/Editor/ClassViewer/Private/SClassViewer.h
@@ -21,8 +21,7 @@ class FMenuBuilder;
 class FTextFilterExpressionEvaluator;
 class UBlueprint;
 class SComboButton;
-class FClassViewerNode;
-class FTextFilterExpressionEvaluator;
+class FClassViewerFilter;
 
 //////////////////////////////////////////////////////////////////////////
 // SClassViewer
@@ -226,12 +225,12 @@ private:
 	/** Init options, cached */
 	FClassViewerInitializationOptions InitOptions;
 
+	/** Filter to use to determine what classes are valid. */
+	TSharedPtr<FClassViewerFilter> ClassFilter;
+
 	/** The items to be displayed in the tree. */
 	TArray< TSharedPtr< FClassViewerNode > > RootTreeItems;
 
-	/** Compiled filter search terms. */
-	TSharedPtr<FTextFilterExpressionEvaluator> TextFilterPtr;
-
 	/** Holds the Slate Tree widget which holds the classes for the Class Viewer. */
 	TSharedPtr<STreeView<TSharedPtr< FClassViewerNode > >> ClassTree;
 
@@ -241,18 +240,6 @@ private:
 	/** The Class Search Box, used for filtering the classes visible. */
 	TSharedPtr<SSearchBox> SearchBox;
 
-	/** true to filter for Actors only. */
-	bool bIsActorsOnly;
-
-	/** true to filter for Placeable classes only. */
-	bool bIsPlaceableOnly;
-
-	/** true to filter for Blueprint Base classes only. */
-	bool bIsBlueprintBaseOnly;
-
-	/** true to filter for unloaded Blueprint classes. */
-	bool bShowUnloadedBlueprints;
-
 	/** true to allow class dynamic loading. */
 	bool bEnableClassDynamicLoading;
 
@@ -283,9 +270,6 @@ private:
 	/** True if we need to set the tree expansion states according to our local copy next tick */
 	bool bPendingSetExpansionStates;
 
-	/** Indicates if the 'Show Internal Classes' option should be enabled or disabled */
-	bool bCanShowInternalClasses;
-
 	/** The button that displays view options */
 	TSharedPtr<SComboButton> ViewOptionsComboButton;
 
diff --git a/Engine/Source/Editor/ClassViewer/Private/UnloadedBlueprintData.cpp b/Engine/Source/Editor/ClassViewer/Private/UnloadedBlueprintData.cpp
--- a/Engine/Source/Editor/ClassViewer/Private/UnloadedBlueprintData.cpp
+++ b/Engine/Source/Editor/ClassViewer/Private/UnloadedBlueprintData.cpp
@@ -116,6 +116,25 @@ const UClass* FUnloadedBlueprintData::GetNativeParent() const
 	return nullptr;
 }
 
+TSharedPtr<FString> FUnloadedBlueprintData::GetClassName() const
+{
+	if (ClassViewerNode.IsValid())
+	{
+		return ClassViewerNode.Pin()->GetClassName();
+	}
+
+	return TSharedPtr<FString>();
+}
+
+FName FUnloadedBlueprintData::GetClassPath() const
+{
+	if (ClassViewerNode.IsValid())
+	{
+		return ClassViewerNode.Pin()->ClassPath;
+	}
+
+	return FName();
+}
 
 const TWeakPtr<FClassViewerNode>& FUnloadedBlueprintData::GetClassViewerNode() const
 {
diff --git a/Engine/Source/Editor/ClassViewer/Private/UnloadedBlueprintData.h b/Engine/Source/Editor/ClassViewer/Private/UnloadedBlueprintData.h
--- a/Engine/Source/Editor/ClassViewer/Private/UnloadedBlueprintData.h
+++ b/Engine/Source/Editor/ClassViewer/Private/UnloadedBlueprintData.h
@@ -12,9 +12,9 @@ public:
 
 	virtual ~FUnloadedBlueprintData() {}
 
-	virtual bool HasAnyClassFlags( uint32 InFlagsToCheck ) const override;
+	virtual bool HasAnyClassFlags(uint32 InFlagsToCheck) const override;
 
-	virtual bool HasAllClassFlags( uint32 InFlagsToCheck ) const override;
+	virtual bool HasAllClassFlags(uint32 InFlagsToCheck) const override;
 
 	virtual void SetClassFlags(uint32 InFlags) override;
 
@@ -24,6 +24,14 @@ public:
 
 	virtual bool IsA(const UClass* InClass) const override;
 
+	virtual void SetNormalBlueprintType(bool bInNormalBPType) override { bNormalBlueprintType = bInNormalBPType; }
+
+	virtual bool IsNormalBlueprintType() const override { return bNormalBlueprintType; }
+
+	virtual TSharedPtr<FString> GetClassName() const override;
+
+	virtual FName GetClassPath() const override;
+
 	virtual const UClass* GetClassWithin() const override;
 
 	virtual const UClass* GetNativeParent() const override;
@@ -38,6 +46,9 @@ private:
 	/** Flags for the class. */
 	uint32 ClassFlags = CLASS_None;
 
+	/** Is this a normal blueprint type? */
+	bool bNormalBlueprintType;
+
 	/** The implemented interfaces for this class. */
 	TArray<FString> ImplementedInterfaces;
 
diff --git a/Engine/Source/Editor/ClassViewer/Public/ClassViewerFilter.h b/Engine/Source/Editor/ClassViewer/Public/ClassViewerFilter.h
--- a/Engine/Source/Editor/ClassViewer/Public/ClassViewerFilter.h
+++ b/Engine/Source/Editor/ClassViewer/Public/ClassViewerFilter.h
@@ -2,8 +2,11 @@
 #pragma once
 
 #include "CoreMinimal.h"
+#include "Settings/ClassViewerSettings.h"
 
+class FClassViewerNode;
 class FClassViewerInitializationOptions;
+class FTextFilterExpressionEvaluator;
 
 /** Interface class for creating filters for the Class Viewer. */
 class IClassViewerFilter
@@ -30,6 +33,24 @@ public:
 	virtual bool IsUnloadedClassAllowed(const FClassViewerInitializationOptions& InInitOptions, const TSharedRef< const class IUnloadedBlueprintData > InUnloadedClassData, TSharedRef< class FClassViewerFilterFuncs > InFilterFuncs) = 0;
 };
 
+/** Filter class that performs many common checks. */
+class FClassViewerFilter : public IClassViewerFilter
+{
+public:
+	FClassViewerFilter(const FClassViewerInitializationOptions& InInitOptions);
+
+	virtual bool IsNodeAllowed(const FClassViewerInitializationOptions& InInitOptions, const TSharedRef<FClassViewerNode>& Node);
+
+	virtual bool IsClassAllowed(const FClassViewerInitializationOptions& InInitOptions, const UClass* InClass, TSharedRef< class FClassViewerFilterFuncs > InFilterFuncs ) override;
+	virtual bool IsUnloadedClassAllowed(const FClassViewerInitializationOptions& InInitOptions, const TSharedRef< const class IUnloadedBlueprintData > InUnloadedClassData, TSharedRef< class FClassViewerFilterFuncs > InFilterFuncs) override;
+
+	TArray<UClass*> InternalClasses;
+	TArray<FDirectoryPath> InternalPaths;
+
+	TSharedRef<FTextFilterExpressionEvaluator> TextFilter;
+	TSharedRef<FClassViewerFilterFuncs> FilterFunctions;
+};
+
 namespace EFilterReturn
 {
 	enum Type{ Failed = 0, Passed, NoItems };
@@ -164,6 +185,7 @@ public:
 class IUnloadedBlueprintData
 {
 public:
+
 	/**
 	 * Used to safely check whether the passed in flag is set.
 	 *
@@ -227,4 +249,24 @@ public:
 	 * @return The child-most Native class in the hierarchy.
 	 */
 	virtual const UClass* GetNativeParent() const = 0;
+
+	/** 
+	 * Set whether or not this blueprint is a normal blueprint.
+	 */
+	virtual void SetNormalBlueprintType(bool bInNormalBPType) = 0;
+
+	/** 
+	 * Get whether or not this blueprint is a normal blueprint. 
+	 */
+	virtual bool IsNormalBlueprintType() const = 0;
+
+	/**
+	 * Get the generated class name of this blueprint.
+	 */
+	virtual TSharedPtr<FString> GetClassName() const = 0;
+
+	/**
+	 * Get the class path of this blueprint.
+	 */
+	virtual FName GetClassPath() const = 0;
 };
diff --git a/Engine/Source/Editor/ClassViewer/Public/ClassViewerModule.h b/Engine/Source/Editor/ClassViewer/Public/ClassViewerModule.h
--- a/Engine/Source/Editor/ClassViewer/Public/ClassViewerModule.h
+++ b/Engine/Source/Editor/ClassViewer/Public/ClassViewerModule.h
@@ -7,6 +7,7 @@
 
 class IClassViewerFilter;
 class IPropertyHandle;
+class FClassViewerFilterFuncs;
 
 /** Delegate used with the Class Viewer in 'class picking' mode.  You'll bind a delegate when the
     class viewer widget is created, which will be fired off when a class is selected in the list */
@@ -163,4 +164,10 @@ public:
 	virtual TSharedRef<class SWidget> CreateClassViewer(const FClassViewerInitializationOptions& InitOptions,
 		const FOnClassPicked& OnClassPickedDelegate );
 
+	/** 
+	 * Create a new class filter from the given initialization options.
+	 */
+	virtual TSharedRef<IClassViewerFilter> CreateClassFilter(const FClassViewerInitializationOptions& InitOptions);
+
+	virtual TSharedRef<FClassViewerFilterFuncs> CreateFilterFuncs();
 };
diff --git a/Engine/Source/Editor/PropertyEditor/Private/UserInterface/PropertyEditor/SPropertyEditorClass.cpp b/Engine/Source/Editor/PropertyEditor/Private/UserInterface/PropertyEditor/SPropertyEditorClass.cpp
--- a/Engine/Source/Editor/PropertyEditor/Private/UserInterface/PropertyEditor/SPropertyEditorClass.cpp
+++ b/Engine/Source/Editor/PropertyEditor/Private/UserInterface/PropertyEditor/SPropertyEditorClass.cpp
@@ -11,7 +11,6 @@
 #include "ClassViewerModule.h"
 #include "ClassViewerFilter.h"
 
-
 #define LOCTEXT_NAMESPACE "PropertyEditor"
 
 class FPropertyEditorClassFilter : public IClassViewerFilter
@@ -26,7 +25,7 @@ public:
 	/** Whether or not abstract classes are allowed. */
 	bool bAllowAbstract;
 
-	bool IsClassAllowed(const FClassViewerInitializationOptions& InInitOptions, const UClass* InClass, TSharedRef< FClassViewerFilterFuncs > InFilterFuncs ) override
+	virtual bool IsClassAllowed(const FClassViewerInitializationOptions& InInitOptions, const UClass* InClass, TSharedRef< FClassViewerFilterFuncs > InFilterFuncs ) override
 	{
 		bool bMatchesFlags = !InClass->HasAnyClassFlags(CLASS_Hidden|CLASS_HideDropDown|CLASS_Deprecated) &&
 			(bAllowAbstract || !InClass->HasAnyClassFlags(CLASS_Abstract));
@@ -84,6 +83,8 @@ bool SPropertyEditorClass::Supports(const TSharedRef< class FPropertyEditor >& I
 void SPropertyEditorClass::Construct(const FArguments& InArgs, const TSharedPtr< class FPropertyEditor >& InPropertyEditor)
 {
 	PropertyEditor = InPropertyEditor;
+
+	CreateClassFilter();
 	
 	if (PropertyEditor.IsValid())
 	{
@@ -129,8 +130,6 @@ void SPropertyEditorClass::Construct(const FArguments& InArgs, const TSharedPtr<
 
 		SelectedClass = InArgs._SelectedClass;
 		OnSetClass = InArgs._OnSetClass;
-
-
 	}
 	
 	SAssignNew(ComboButton, SComboButton)
@@ -199,29 +198,36 @@ FText SPropertyEditorClass::GetDisplayValueAsString() const
 	
 }
 
-TSharedRef<SWidget> SPropertyEditorClass::GenerateClassPicker()
+void SPropertyEditorClass::CreateClassFilter()
 {
-	FClassViewerInitializationOptions Options;
-	Options.bShowBackgroundBorder = false;
-	Options.bShowUnloadedBlueprints = true;
-	Options.bShowNoneOption = bAllowNone; 
+	ClassViewerOptions.bShowBackgroundBorder = false;
+	ClassViewerOptions.bShowUnloadedBlueprints = true;
+	ClassViewerOptions.bShowNoneOption = bAllowNone;
 
-	if(PropertyEditor.IsValid())
+	if (PropertyEditor.IsValid())
 	{
-		Options.PropertyHandle = PropertyEditor->GetPropertyHandle();
+		ClassViewerOptions.PropertyHandle = PropertyEditor->GetPropertyHandle();
 	}
 
-	TSharedPtr<FPropertyEditorClassFilter> ClassFilter = MakeShareable(new FPropertyEditorClassFilter);
-	Options.ClassFilter = ClassFilter;
-	ClassFilter->ClassPropertyMetaClass = MetaClass;
-	ClassFilter->InterfaceThatMustBeImplemented = RequiredInterface;
-	ClassFilter->bAllowAbstract = bAllowAbstract;
-	Options.bIsBlueprintBaseOnly = bIsBlueprintBaseOnly;
-	Options.bIsPlaceableOnly = bAllowOnlyPlaceable;
-	Options.NameTypeToDisplay = (bShowDisplayNames ? EClassViewerNameTypeToDisplay::DisplayName : EClassViewerNameTypeToDisplay::ClassName);
-	Options.DisplayMode = bShowTree ? EClassViewerDisplayMode::TreeView : EClassViewerDisplayMode::ListView;
-	Options.bAllowViewOptions = bShowViewOptions;
+	ClassViewerOptions.bIsBlueprintBaseOnly = bIsBlueprintBaseOnly;
+	ClassViewerOptions.bIsPlaceableOnly = bAllowOnlyPlaceable;
+	ClassViewerOptions.NameTypeToDisplay = (bShowDisplayNames ? EClassViewerNameTypeToDisplay::DisplayName : EClassViewerNameTypeToDisplay::ClassName);
+	ClassViewerOptions.DisplayMode = bShowTree ? EClassViewerDisplayMode::TreeView : EClassViewerDisplayMode::ListView;
+	ClassViewerOptions.bAllowViewOptions = bShowViewOptions;
 
+	TSharedPtr<FPropertyEditorClassFilter> PropEdClassFilter = MakeShareable(new FPropertyEditorClassFilter);
+	ClassViewerOptions.ClassFilter = PropEdClassFilter;
+
+	PropEdClassFilter->ClassPropertyMetaClass = MetaClass;
+	PropEdClassFilter->InterfaceThatMustBeImplemented = RequiredInterface;
+	PropEdClassFilter->bAllowAbstract = bAllowAbstract;
+
+	ClassFilter = FModuleManager::LoadModuleChecked<FClassViewerModule>("ClassViewer").CreateClassFilter(ClassViewerOptions);
+	ClassFilterFuncs = FModuleManager::LoadModuleChecked<FClassViewerModule>("ClassViewer").CreateFilterFuncs();
+}
+
+TSharedRef<SWidget> SPropertyEditorClass::GenerateClassPicker()
+{
 	FOnClassPicked OnPicked(FOnClassPicked::CreateRaw(this, &SPropertyEditorClass::OnClassPicked));
 
 	return SNew(SBox)
@@ -232,7 +238,7 @@ TSharedRef<SWidget> SPropertyEditorClass::GenerateClassPicker()
 			.AutoHeight()
 			.MaxHeight(500)
 			[
-				FModuleManager::LoadModuleChecked<FClassViewerModule>("ClassViewer").CreateClassViewer(Options, OnPicked)
+				FModuleManager::LoadModuleChecked<FClassViewerModule>("ClassViewer").CreateClassViewer(ClassViewerOptions, OnPicked)
 			]			
 		];
 }
@@ -273,51 +279,60 @@ void SPropertyEditorClass::SendToObjects(const FString& NewValue)
 	}
 }
 
-void SPropertyEditorClass::OnDragEnter(const FGeometry& MyGeometry, const FDragDropEvent& DragDropEvent)
+static UObject* LoadDragDropObject(TSharedPtr<FAssetDragDropOp> UnloadedClassOp)
 {
-	TSharedPtr<FAssetDragDropOp> UnloadedClassOp = DragDropEvent.GetOperationAs<FAssetDragDropOp>();
-	if (UnloadedClassOp.IsValid())
+	FString AssetPath;
+
+	// Find the class/blueprint path
+	if (UnloadedClassOp->HasAssets())
 	{
-		bool bAllAssetWereLoaded = true;
+		AssetPath = UnloadedClassOp->GetAssets()[0].ObjectPath.ToString();
+	}
+	else if (UnloadedClassOp->HasAssetPaths())
+	{
+		AssetPath = UnloadedClassOp->GetAssetPaths()[0];
+	}
 
-		FString AssetPath;
-		FString PathName;
+	// Check to see if the asset can be found, otherwise load it.
+	UObject* Object = FindObject<UObject>(nullptr, *AssetPath);
+	if (Object == nullptr)
+	{
+		// Load the package.
+		GWarn->BeginSlowTask(LOCTEXT("OnDrop_LoadPackage", "Fully Loading Package For Drop"), true, false);
 
-		// Find the class/blueprint path
-		if (UnloadedClassOp->HasAssets())
-		{
-			AssetPath = UnloadedClassOp->GetAssets()[0].ObjectPath.ToString();
-		}
-		else if (UnloadedClassOp->HasAssetPaths())
-		{
-			AssetPath = UnloadedClassOp->GetAssetPaths()[0];
-		}
+		Object = LoadObject<UObject>(nullptr, *AssetPath);
 
-		// Check to see if the asset can be found, otherwise load it.
-		UObject* Object = FindObject<UObject>(nullptr, *AssetPath);
-		if (Object == nullptr)
-		{
-			// Load the package.
-			GWarn->BeginSlowTask(LOCTEXT("OnDrop_LoadPackage", "Fully Loading Package For Drop"), true, false);
+		GWarn->EndSlowTask();
+	}
 
-			Object = LoadObject<UObject>(nullptr, *AssetPath);
+	return Object;
+}
 
-			GWarn->EndSlowTask();
-		}
+void SPropertyEditorClass::OnDragEnter(const FGeometry& MyGeometry, const FDragDropEvent& DragDropEvent)
+{
+	TSharedPtr<FAssetDragDropOp> UnloadedClassOp = DragDropEvent.GetOperationAs<FAssetDragDropOp>();
+	if (UnloadedClassOp.IsValid())
+	{
+		UObject* Object = LoadDragDropObject(UnloadedClassOp);
+
+		bool bOK = false;
 
 		if (UClass* Class = Cast<UClass>(Object))
 		{
-			// This was pointing to a class directly
-			UnloadedClassOp->SetToolTip(FText::GetEmpty(), FEditorStyle::GetBrush(TEXT("Graph.ConnectorFeedback.OK")));
+			bOK = ClassFilter->IsClassAllowed(ClassViewerOptions, Class, ClassFilterFuncs.ToSharedRef());
 		}
 		else if (UBlueprint* Blueprint = Cast<UBlueprint>(Object))
 		{
 			if (Blueprint->GeneratedClass)
 			{
-				// This was pointing to a blueprint, get generated class
-				UnloadedClassOp->SetToolTip(FText::GetEmpty(), FEditorStyle::GetBrush(TEXT("Graph.ConnectorFeedback.OK")));
+				bOK = ClassFilter->IsClassAllowed(ClassViewerOptions, Blueprint->GeneratedClass, ClassFilterFuncs.ToSharedRef());
 			}
 		}
+		
+		if (bOK)
+		{
+			UnloadedClassOp->SetToolTip(FText::GetEmpty(), FEditorStyle::GetBrush(TEXT("Graph.ConnectorFeedback.OK")));
+		}
 		else
 		{
 			UnloadedClassOp->SetToolTip(FText::GetEmpty(), FEditorStyle::GetBrush(TEXT("Graph.ConnectorFeedback.Error")));
@@ -351,43 +366,25 @@ FReply SPropertyEditorClass::OnDrop(const FGeometry& MyGeometry, const FDragDrop
 	TSharedPtr<FAssetDragDropOp> UnloadedClassOp = DragDropEvent.GetOperationAs<FAssetDragDropOp>();
 	if (UnloadedClassOp.IsValid())
 	{
-		bool bAllAssetWereLoaded = true;
-
-		FString AssetPath;
-
-		// Find the class/blueprint path
-		if (UnloadedClassOp->HasAssets())
-		{
-			AssetPath = UnloadedClassOp->GetAssets()[0].ObjectPath.ToString();
-		}
-		else if (UnloadedClassOp->HasAssetPaths())
-		{
-			AssetPath = UnloadedClassOp->GetAssetPaths()[0];
-		}
-
-		// Check to see if the asset can be found, otherwise load it.
-		UObject* Object = FindObject<UObject>(nullptr, *AssetPath);
-		if(Object == nullptr)
-		{
-			// Load the package.
-			GWarn->BeginSlowTask(LOCTEXT("OnDrop_LoadPackage", "Fully Loading Package For Drop"), true, false);
-
-			Object = LoadObject<UObject>(nullptr, *AssetPath);
-
-			GWarn->EndSlowTask();
-		}
+		UObject* Object = LoadDragDropObject(UnloadedClassOp);
 
 		if (UClass* Class = Cast<UClass>(Object))
 		{
-			// This was pointing to a class directly
-			SendToObjects(Class->GetPathName());
+			if (ClassFilter->IsClassAllowed(ClassViewerOptions, Class, ClassFilterFuncs.ToSharedRef()))
+			{
+				// This was pointing to a class directly
+				SendToObjects(Class->GetPathName());
+			}
 		}
 		else if (UBlueprint* Blueprint = Cast<UBlueprint>(Object))
 		{
 			if (Blueprint->GeneratedClass)
 			{
-				// This was pointing to a blueprint, get generated class
-				SendToObjects(Blueprint->GeneratedClass->GetPathName());
+				if (ClassFilter->IsClassAllowed(ClassViewerOptions, Blueprint->GeneratedClass, ClassFilterFuncs.ToSharedRef()))
+				{
+					// This was pointing to a blueprint, get generated class
+					SendToObjects(Blueprint->GeneratedClass->GetPathName());
+				}
 			}
 		}
 
diff --git a/Engine/Source/Editor/PropertyEditor/Private/UserInterface/PropertyEditor/SPropertyEditorClass.h b/Engine/Source/Editor/PropertyEditor/Private/UserInterface/PropertyEditor/SPropertyEditorClass.h
--- a/Engine/Source/Editor/PropertyEditor/Private/UserInterface/PropertyEditor/SPropertyEditorClass.h
+++ b/Engine/Source/Editor/PropertyEditor/Private/UserInterface/PropertyEditor/SPropertyEditorClass.h
@@ -14,6 +14,10 @@
 #include "Presentation/PropertyEditor/PropertyEditor.h"
 #include "UserInterface/PropertyEditor/PropertyEditorConstants.h"
 #include "PropertyCustomizationHelpers.h"
+#include "ClassViewerModule.h"
+
+class IClassViewerFilter;
+class FClassViewerFilterFuncs;
 
 /**
  * A widget used to edit Class properties (UClass type properties).
@@ -69,6 +73,8 @@ private:
 	virtual void OnDragLeave(const FDragDropEvent& DragDropEvent) override;
 	virtual FReply OnDrop(const FGeometry& MyGeometry, const FDragDropEvent& DragDropEvent) override;
 
+	TSharedRef<SWidget> ConstructClassViewer();
+
 	/** 
 	 * Generates a class picker with a filter to show only classes allowed to be selected. 
 	 *
@@ -95,6 +101,15 @@ private:
 	/** Used when the property deals with Classes and will display a Class Picker. */
 	TSharedPtr<class SComboButton> ComboButton;
 
+	/** Class filter that the class viewer is using. */
+	TSharedPtr<IClassViewerFilter> ClassFilter;
+
+	/** Filter functions for class viewer. */
+	TSharedPtr<FClassViewerFilterFuncs> ClassFilterFuncs;
+
+	/** Options used for creating the class viewer. */
+	FClassViewerInitializationOptions ClassViewerOptions;
+
 	/** The meta class that the selected class must be a child-of */
 	const UClass* MetaClass;
 	/** An interface that the selected class must implement */
@@ -118,4 +133,6 @@ private:
 	TAttribute<const UClass*> SelectedClass;
 	/** Delegate used to set the currently selected class (required if PropertyEditor == null) */
 	FOnSetClass OnSetClass;
+
+	void CreateClassFilter();
 };
