Removed optimization pragma.

diff --git a/Engine/Source/Editor/ClassViewer/Private/ClassViewerFilter.cpp b/Engine/Source/Editor/ClassViewer/Private/ClassViewerFilter.cpp
--- a/Engine/Source/Editor/ClassViewer/Private/ClassViewerFilter.cpp
+++ b/Engine/Source/Editor/ClassViewer/Private/ClassViewerFilter.cpp
@@ -402,8 +402,6 @@ bool FClassViewerFilter::IsNodeAllowed(const FClassViewerInitializationOptions&
 	return false;
 }
 
-PRAGMA_DISABLE_OPTIMIZATION
-
 bool FClassViewerFilter::IsClassAllowed(const FClassViewerInitializationOptions& InInitOptions, const UClass* InClass, TSharedRef<class FClassViewerFilterFuncs> InFilterFuncs)
 {
 	if (InInitOptions.bIsActorsOnly && !InClass->IsChildOf(AActor::StaticClass()))
