Marking function static to fix builder error.

diff --git a/Engine/Source/Editor/ClassViewer/Private/ClassViewerFilter.cpp b/Engine/Source/Editor/ClassViewer/Private/ClassViewerFilter.cpp
--- a/Engine/Source/Editor/ClassViewer/Private/ClassViewerFilter.cpp
+++ b/Engine/Source/Editor/ClassViewer/Private/ClassViewerFilter.cpp
@@ -316,7 +316,7 @@ static bool IsPlaceable(const TSharedRef<const IUnloadedBlueprintData>& InBluepr
  * @param InClass	The class to verify can be made into a Blueprint
  * @return			true if the class can be made into a Blueprint
  */
-bool CanCreateBlueprintOfClass_IgnoreDeprecation(UClass* InClass)
+static bool CanCreateBlueprintOfClass_IgnoreDeprecation(UClass* InClass)
 {
 	// Temporarily remove the deprecated flag so we can check if it is valid for
 	bool bIsClassDeprecated = InClass->HasAnyClassFlags(CLASS_Deprecated);
