Fixed builder error related to an ambiguous function call.

diff --git a/Engine/Source/Editor/ClassViewer/Private/ClassViewerFilter.cpp b/Engine/Source/Editor/ClassViewer/Private/ClassViewerFilter.cpp
--- a/Engine/Source/Editor/ClassViewer/Private/ClassViewerFilter.cpp
+++ b/Engine/Source/Editor/ClassViewer/Private/ClassViewerFilter.cpp
@@ -316,7 +316,7 @@ static bool IsPlaceable(const TSharedRef<const IUnloadedBlueprintData>& InBluepr
  * @param InClass	The class to verify can be made into a Blueprint
  * @return			true if the class can be made into a Blueprint
  */
-static bool CanCreateBlueprintOfClass_IgnoreDeprecation(UClass* InClass)
+static bool CanCreateBlueprintOfClass(UClass* InClass)
 {
 	// Temporarily remove the deprecated flag so we can check if it is valid for
 	bool bIsClassDeprecated = InClass->HasAnyClassFlags(CLASS_Deprecated);
@@ -409,7 +409,7 @@ bool FClassViewerFilter::IsClassAllowed(const FClassViewerInitializationOptions&
 		return false;
 	}
 
-	const bool bPassesBlueprintBaseFilter = !InInitOptions.bIsBlueprintBaseOnly || CanCreateBlueprintOfClass_IgnoreDeprecation(const_cast<UClass*>(InClass));
+	const bool bPassesBlueprintBaseFilter = !InInitOptions.bIsBlueprintBaseOnly || CanCreateBlueprintOfClass(const_cast<UClass*>(InClass));
 	const bool bPassesEditorClassFilter = !InInitOptions.bEditorClassesOnly || IsEditorOnlyObject(InClass);
 
 	// Determine if we allow any developer folder classes, if so determine if this class is in one of the allowed developer folders.
