More missed includes failing non-unity build.

diff --git a/Engine/Source/Editor/ClassViewer/Private/ClassViewerFilter.cpp b/Engine/Source/Editor/ClassViewer/Private/ClassViewerFilter.cpp
--- a/Engine/Source/Editor/ClassViewer/Private/ClassViewerFilter.cpp
+++ b/Engine/Source/Editor/ClassViewer/Private/ClassViewerFilter.cpp
@@ -1,7 +1,10 @@
 #include "ClassViewerFilter.h"
 
 #include "UnloadedBlueprintData.h"
+#include "Engine/Brush.h"
 #include "Kismet2/KismetEditorUtilities.h"
+#include "Misc/PackageName.h"
+#include "Misc/Paths.h"
 #include "Misc/TextFilterExpressionEvaluator.h"
 
 EFilterReturn::Type FClassViewerFilterFuncs::IfInChildOfClassesSet(TSet< const UClass* >& InSet, const UClass* InClass)
diff --git a/Engine/Source/Editor/ClassViewer/Public/ClassViewerFilter.h b/Engine/Source/Editor/ClassViewer/Public/ClassViewerFilter.h
--- a/Engine/Source/Editor/ClassViewer/Public/ClassViewerFilter.h
+++ b/Engine/Source/Editor/ClassViewer/Public/ClassViewerFilter.h
@@ -2,6 +2,7 @@
 #pragma once
 
 #include "CoreMinimal.h"
+#include "Engine/EngineTypes.h"
 #include "Settings/ClassViewerSettings.h"
 
 class FClassViewerNode;
diff --git a/Engine/Source/Editor/PropertyEditor/Private/EditConditionParser.h b/Engine/Source/Editor/PropertyEditor/Private/EditConditionParser.h
--- a/Engine/Source/Editor/PropertyEditor/Private/EditConditionParser.h
+++ b/Engine/Source/Editor/PropertyEditor/Private/EditConditionParser.h
@@ -2,7 +2,10 @@
 
 #pragma once
 
+#include "CoreTypes.h"
 #include "Misc/ExpressionParserTypes.h"
+#include "UObject/WeakObjectPtr.h"
+#include "UObject/WeakObjectPtrTemplates.h"
 
 namespace EditConditionParserNamespace
 {
