Missed includes.

diff --git a/Engine/Source/Editor/ClassViewer/Private/ClassViewerFilter.cpp b/Engine/Source/Editor/ClassViewer/Private/ClassViewerFilter.cpp
--- a/Engine/Source/Editor/ClassViewer/Private/ClassViewerFilter.cpp
+++ b/Engine/Source/Editor/ClassViewer/Private/ClassViewerFilter.cpp
@@ -3,6 +3,7 @@
 #include "UnloadedBlueprintData.h"
 #include "Engine/Brush.h"
 #include "Kismet2/KismetEditorUtilities.h"
+#include "Misc/ConfigCacheIni.h"
 #include "Misc/PackageName.h"
 #include "Misc/Paths.h"
 #include "Misc/TextFilterExpressionEvaluator.h"
diff --git a/Engine/Source/Editor/PropertyEditor/Private/EditConditionParser.cpp b/Engine/Source/Editor/PropertyEditor/Private/EditConditionParser.cpp
--- a/Engine/Source/Editor/PropertyEditor/Private/EditConditionParser.cpp
+++ b/Engine/Source/Editor/PropertyEditor/Private/EditConditionParser.cpp
@@ -1,6 +1,8 @@
 // Copyright 1998-2019 Epic Games, Inc. All Rights Reserved.
 
 #include "EditConditionParser.h"
+
+#include "Math/BasicMathExpressionEvaluator.h"
 #include "Misc/ExpressionParser.h"
 
 #include "PropertyNode.h"
