Optimize FClassHierarchy::PopulateClassHierarchy (~10s to <1s).

#rnx

diff --git a/Engine/Source/Editor/ClassViewer/Private/SClassViewer.cpp b/Engine/Source/Editor/ClassViewer/Private/SClassViewer.cpp
--- a/Engine/Source/Editor/ClassViewer/Private/SClassViewer.cpp
+++ b/Engine/Source/Editor/ClassViewer/Private/SClassViewer.cpp
@@ -149,9 +149,9 @@ public:
 private:
 	/** Recursive function to build a tree, will not filter.
 	 *	@param InOutRootNode						The node that this function will add the children of to the tree.
-	 *  @param PackageNameToAssetDataMap			The asset registry map of blueprint package names to blueprint data
+	 *  @param InOutClassPathToNode			        Map that will be populated with the class path to the corresponding class node.
 	 */	
-	void AddChildren_NoFilter( TSharedPtr< FClassViewerNode >& InOutRootNode, const TMultiMap<FName,FAssetData>& BlueprintPackageToAssetDataMap );
+	void AddChildren_NoFilter( TSharedPtr< FClassViewerNode >& InOutRootNode, TMap<FName, TSharedPtr< FClassViewerNode >>& InOutClassPathToNode );
 
 	/** Called when hot reload has finished */
 	void OnHotReload( bool bWasTriggeredAutomatically );
@@ -1049,7 +1049,7 @@ void FClassHierarchy::OnHotReload( bool bWasTriggeredAutomatically )
 	ClassViewer::Helpers::RequestPopulateClassHierarchy();
 }
 
-void FClassHierarchy::AddChildren_NoFilter( TSharedPtr< FClassViewerNode >& InOutRootNode, const TMultiMap<FName, FAssetData>& BlueprintPackageToAssetDataMap )
+void FClassHierarchy::AddChildren_NoFilter( TSharedPtr< FClassViewerNode >& InOutRootNode, TMap<FName, TSharedPtr< FClassViewerNode >>& InOutClassPathToNode )
 {
 	UClass* RootClass = UObject::StaticClass();
 
@@ -1088,12 +1088,14 @@ void FClassHierarchy::AddChildren_NoFilter( TSharedPtr< FClassViewerNode >& InOu
 				if ( !ParentEntry.IsValid() )
 				{
 					ParentEntry = MakeShareable(new FClassViewerNode(CurrentClass->GetSuperClass()));
+                    InOutClassPathToNode.Add(ParentEntry->ClassPath, ParentEntry);
 				}
 
 				TSharedPtr<FClassViewerNode>& MyEntry = Nodes.FindOrAdd(CurrentClass);
 				if ( !MyEntry.IsValid() )
 				{
 					MyEntry = MakeShareable(new FClassViewerNode(CurrentClass));
+                    InOutClassPathToNode.Add(MyEntry->ClassPath, MyEntry);
 				}
 
 				if ( !Visited.Contains(CurrentClass) )
@@ -1397,21 +1399,21 @@ void FClassHierarchy::PopulateClassHierarchy()
 	Filter.bRecursiveClasses = true;
 	AssetRegistryModule.Get().GetAssets(Filter, BlueprintList);
 
-	TMultiMap<FName,FAssetData> BlueprintPackageToAssetDataMap;
+	TMap<FName, TSharedPtr< FClassViewerNode >> ClassPathToNode;
 	for ( int32 AssetIdx = 0; AssetIdx < BlueprintList.Num(); ++AssetIdx )
 	{
 		TSharedPtr< FClassViewerNode > NewNode;
 		LoadUnloadedTagData(NewNode, BlueprintList[AssetIdx]);
 		RootLevelClasses.Add(NewNode);
 
+		check(!NewNode->GetChildrenList().Num());
+		ClassPathToNode.Add(NewNode->ClassPath, NewNode);
+
 		// Find the blueprint if it's loaded.
 		FindClass(NewNode);
-
-
-		BlueprintPackageToAssetDataMap.Add(BlueprintList[AssetIdx].PackageName, BlueprintList[AssetIdx]);
 	}
 
-	AddChildren_NoFilter(ObjectClassRoot, BlueprintPackageToAssetDataMap);
+	AddChildren_NoFilter(ObjectClassRoot, ClassPathToNode);
 
 	RootLevelClasses.Add(ObjectClassRoot);
 
@@ -1423,18 +1425,15 @@ void FClassHierarchy::PopulateClassHierarchy()
 			// Resolve the parent's class name locally and use it to find the parent's class.
 			FString ParentClassPath = RootLevelClasses[CurrentNodeIdx]->ParentClassPath.ToString();
 			const UClass* ParentClass = FindObject<UClass>(nullptr, *ParentClassPath);
+			TSharedPtr< FClassViewerNode > ParentNode;
 
-			for (int32 SearchNodeIdx = 0; SearchNodeIdx < RootLevelClasses.Num(); ++SearchNodeIdx)
+			if (TSharedPtr< FClassViewerNode >* ParentNodePtr = ClassPathToNode.Find(RootLevelClasses[CurrentNodeIdx]->ParentClassPath))
 			{
-				TSharedPtr< FClassViewerNode > ParentNode = FindParent(RootLevelClasses[SearchNodeIdx], RootLevelClasses[CurrentNodeIdx]->ParentClassPath, ParentClass);
-				if(ParentNode.IsValid())
-				{
-					// AddUniqueChild makes sure that when a node was generated one by EditorClassHierarchy and one from LoadUnloadedTagData - the proper one is selected
-					ParentNode->AddUniqueChild(RootLevelClasses[CurrentNodeIdx]);	
-					RootLevelClasses.RemoveAtSwap(CurrentNodeIdx);
-					--CurrentNodeIdx;
-					break;
-				}
+				// AddUniqueChild makes sure that when a node was generated one by EditorClassHierarchy and one from LoadUnloadedTagData - the proper one is selected
+				ParentNode = *ParentNodePtr;
+				ParentNode->AddUniqueChild(RootLevelClasses[CurrentNodeIdx]);	
+				RootLevelClasses.RemoveAtSwap(CurrentNodeIdx);
+				--CurrentNodeIdx;
 			}
 
 		}
