The search filter (SClassViewer) now allows selection of the parents of the found objects. Previously, only the elements found by the filter could be selected, even though their parents were visualized as well.

diff --git a/Engine/Source/Editor/ClassViewer/Private/ClassViewerFilter.cpp b/Engine/Source/Editor/ClassViewer/Private/ClassViewerFilter.cpp
--- a/Engine/Source/Editor/ClassViewer/Private/ClassViewerFilter.cpp
+++ b/Engine/Source/Editor/ClassViewer/Private/ClassViewerFilter.cpp
@@ -415,21 +415,27 @@ FClassViewerFilter::FClassViewerFilter(const FClassViewerInitializationOptions&
 	}
 }
 
-bool FClassViewerFilter::IsNodeAllowed(const FClassViewerInitializationOptions& InInitOptions, const TSharedRef<FClassViewerNode>& InNode)
+bool FClassViewerFilter::IsNodeAllowed(const FClassViewerInitializationOptions& InInitOptions, const TSharedRef<FClassViewerNode>& InNode, const bool bCheckTextFilter)
 {
 	if (InNode->Class.IsValid())
 	{
-		return IsClassAllowed(InInitOptions, InNode->Class.Get(), FilterFunctions);
+		return IsClassAllowed(InInitOptions, InNode->Class.Get(), FilterFunctions, bCheckTextFilter);
 	}
 	else if (InInitOptions.bShowUnloadedBlueprints && InNode->UnloadedBlueprintData.IsValid())
 	{
-		return IsUnloadedClassAllowed(InInitOptions, InNode->UnloadedBlueprintData.ToSharedRef(), FilterFunctions);
+		return IsUnloadedClassAllowed(InInitOptions, InNode->UnloadedBlueprintData.ToSharedRef(), FilterFunctions, bCheckTextFilter);
 	}
 
 	return false;
 }
 
 bool FClassViewerFilter::IsClassAllowed(const FClassViewerInitializationOptions& InInitOptions, const UClass* InClass, TSharedRef<class FClassViewerFilterFuncs> InFilterFuncs)
+{
+	const bool bCheckTextFilter = true;
+	return IsClassAllowed(InInitOptions, InClass, InFilterFuncs, bCheckTextFilter);
+}
+
+bool FClassViewerFilter::IsClassAllowed(const FClassViewerInitializationOptions& InInitOptions, const UClass* InClass, TSharedRef<class FClassViewerFilterFuncs> InFilterFuncs, const bool bCheckTextFilter)
 {
 	if (InInitOptions.bIsActorsOnly && !InClass->IsChildOf(AActor::StaticClass()))
 	{
@@ -508,12 +514,18 @@ bool FClassViewerFilter::IsClassAllowed(const FClassViewerInitializationOptions&
 
 	bool bPassesFilter = bPassesPlaceableFilter && bPassesBlueprintBaseFilter 
 		&& bPassesDeveloperFilter && bPassesInternalFilter && bPassesEditorClassFilter 
-		&& bPassesCustomFilter && bPassesTextFilter && bPassesAssetReferenceFilter;
+		&& bPassesCustomFilter && (!bCheckTextFilter || bPassesTextFilter) && bPassesAssetReferenceFilter;
 
 	return bPassesFilter;
 }
 
 bool FClassViewerFilter::IsUnloadedClassAllowed(const FClassViewerInitializationOptions& InInitOptions, const TSharedRef<const IUnloadedBlueprintData> InUnloadedClassData, TSharedRef<FClassViewerFilterFuncs> InFilterFuncs)
+{
+	const bool bCheckTextFilter = true;
+	return IsUnloadedClassAllowed(InInitOptions, InUnloadedClassData, InFilterFuncs, bCheckTextFilter);
+}
+
+bool FClassViewerFilter::IsUnloadedClassAllowed(const FClassViewerInitializationOptions& InInitOptions, const TSharedRef<const IUnloadedBlueprintData> InUnloadedClassData, TSharedRef<FClassViewerFilterFuncs> InFilterFuncs, const bool bCheckTextFilter)
 {
 	if (InInitOptions.bIsActorsOnly && !InUnloadedClassData->IsChildOf(AActor::StaticClass()))
 	{
@@ -585,7 +597,7 @@ bool FClassViewerFilter::IsUnloadedClassAllowed(const FClassViewerInitialization
 
 	bool bPassesFilter = bPassesPlaceableFilter && bPassesBlueprintBaseFilter
 		&& bPassesDeveloperFilter && bPassesInternalFilter && bPassesEditorClassFilter
-		&& bPassesCustomFilter && bPassesTextFilter && bPassesAssetReferenceFilter;
+		&& bPassesCustomFilter && (!bCheckTextFilter || bPassesTextFilter) && bPassesAssetReferenceFilter;
 
 	return bPassesFilter;
 }
diff --git a/Engine/Source/Editor/ClassViewer/Private/ClassViewerNode.cpp b/Engine/Source/Editor/ClassViewer/Private/ClassViewerNode.cpp
--- a/Engine/Source/Editor/ClassViewer/Private/ClassViewerNode.cpp
+++ b/Engine/Source/Editor/ClassViewer/Private/ClassViewerNode.cpp
@@ -30,6 +30,7 @@ FClassViewerNode::FClassViewerNode(UClass* InClass)
 	}
 
 	bPassesFilter = false;
+	bPassesFilterRegardlessTextFilter = false;
 }
 
 FClassViewerNode::FClassViewerNode(const FString& InClassName, const FString& InClassDisplayName)
@@ -37,6 +38,7 @@ FClassViewerNode::FClassViewerNode(const FString& InClassName, const FString& In
 	ClassName = MakeShareable(new FString(InClassName));
 	ClassDisplayName = MakeShareable(new FString(InClassDisplayName));
 	bPassesFilter = false;
+	bPassesFilterRegardlessTextFilter = false;
 
 	Class = nullptr;
 	Blueprint = nullptr;
@@ -47,6 +49,7 @@ FClassViewerNode::FClassViewerNode( const FClassViewerNode& InCopyObject)
 	ClassName = InCopyObject.ClassName;
 	ClassDisplayName = InCopyObject.ClassDisplayName;
 	bPassesFilter = InCopyObject.bPassesFilter;
+	bPassesFilterRegardlessTextFilter = InCopyObject.bPassesFilterRegardlessTextFilter;
 
 	Class = InCopyObject.Class;
 	Blueprint = InCopyObject.Blueprint;
diff --git a/Engine/Source/Editor/ClassViewer/Private/ClassViewerNode.h b/Engine/Source/Editor/ClassViewer/Private/ClassViewerNode.h
--- a/Engine/Source/Editor/ClassViewer/Private/ClassViewerNode.h
+++ b/Engine/Source/Editor/ClassViewer/Private/ClassViewerNode.h
@@ -103,6 +103,12 @@ public:
 	/** true if the class passed the filter. */
 	bool bPassesFilter;
 
+	/** 
+	 * true if the class passed all the sub-filters of the filter (regardless of the TextFilter one).
+	 * This could be useful to verify e.g., that the parent class of a IsNodeAllowed() object is also valid (even though that parent will not likely pass the TextFilter).
+	 */
+	bool bPassesFilterRegardlessTextFilter;
+
 	/** Pointer to the parent to this object. */
 	TWeakPtr< FClassViewerNode > ParentNode;
 
diff --git a/Engine/Source/Editor/ClassViewer/Private/SClassViewer.cpp b/Engine/Source/Editor/ClassViewer/Private/SClassViewer.cpp
--- a/Engine/Source/Editor/ClassViewer/Private/SClassViewer.cpp
+++ b/Engine/Source/Editor/ClassViewer/Private/SClassViewer.cpp
@@ -305,20 +305,24 @@ namespace ClassViewer
 		 *	@param InInitOptions						The class viewer's options, holds the AllowedClasses and DisallowedClasses.
 		 *
 		 *	@return Returns true if the child passed the filter.
-		 */		
+		 */
 		static bool AddChildren_Tree(TSharedPtr< FClassViewerNode >& InOutRootNode, const TSharedPtr< FClassViewerNode >& InOriginalRootNode, 
 			const TSharedPtr<FClassViewerFilter>& InClassFilter, const FClassViewerInitializationOptions& InInitOptions)
 		{
-			InOutRootNode->bPassesFilter = InClassFilter->IsNodeAllowed(InInitOptions, InOutRootNode.ToSharedRef());
+			bool bCheckTextFilter = true;
+			InOutRootNode->bPassesFilter = InClassFilter->IsNodeAllowed(InInitOptions, InOutRootNode.ToSharedRef(), bCheckTextFilter);
 
 			bool bReturnPassesFilter = InOutRootNode->bPassesFilter;
 
+			bCheckTextFilter = false;
+			InOutRootNode->bPassesFilterRegardlessTextFilter = bReturnPassesFilter || InClassFilter->IsNodeAllowed(InInitOptions, InOutRootNode.ToSharedRef(), bCheckTextFilter);
+
 			TArray< TSharedPtr< FClassViewerNode > >& ChildList = InOriginalRootNode->GetChildrenList();
 			for(int32 ChildIdx = 0; ChildIdx < ChildList.Num(); ChildIdx++)
 			{
 				TSharedPtr< FClassViewerNode > NewNode = MakeShareable( new FClassViewerNode( *ChildList[ChildIdx].Get() ) );
 
-				bool bChildrenPassesFilter = AddChildren_Tree(NewNode, ChildList[ChildIdx], InClassFilter, InInitOptions);
+				const bool bChildrenPassesFilter = AddChildren_Tree(NewNode, ChildList[ChildIdx], InClassFilter, InInitOptions);
 				bReturnPassesFilter |= bChildrenPassesFilter;
 
 				if (bChildrenPassesFilter)
@@ -373,7 +377,8 @@ namespace ClassViewer
 		static void AddChildren_List(TArray< TSharedPtr< FClassViewerNode > >& InOutNodeList, const TSharedPtr< FClassViewerNode >& InOriginalRootNode, 
 			const TSharedPtr< FClassViewerFilter >& InClassFilter, const FClassViewerInitializationOptions& InInitOptions)
 		{
-			if (InClassFilter->IsNodeAllowed(InInitOptions, InOriginalRootNode.ToSharedRef()))
+			const bool bCheckTextFilter = true;
+			if (InClassFilter->IsNodeAllowed(InInitOptions, InOriginalRootNode.ToSharedRef(), bCheckTextFilter))
 			{
 				TSharedPtr< FClassViewerNode > NewNode = MakeShareable(new FClassViewerNode(*InOriginalRootNode.Get()));
 				NewNode->bPassesFilter = true;
@@ -407,7 +412,8 @@ namespace ClassViewer
 			// If the option to see the object root class is set, add it to the list, proceed normally from there so the actor's only filter continues to work.
 			if (InInitOptions.bShowObjectRootClass)
 			{
-				if (InClassFilter->IsNodeAllowed(InInitOptions, ObjectClassRoot.ToSharedRef()))
+				const bool bCheckTextFilter = true;
+				if (InClassFilter->IsNodeAllowed(InInitOptions, ObjectClassRoot.ToSharedRef(), bCheckTextFilter))
 				{
 					TSharedPtr< FClassViewerNode > NewNode = MakeShareable(new FClassViewerNode(*ObjectClassRoot.Get()));
 					NewNode->bPassesFilter = true;
@@ -1541,7 +1547,7 @@ void SClassViewer::Construct(const FArguments& InArgs, const FClassViewerInitial
 			.DefaultLabel(NSLOCTEXT("ClassViewer", "Class", "Class"))
 		);
 
-		SAssignNew(ClassTree, STreeView<TSharedPtr< FClassViewerNode > >)
+	SAssignNew(ClassTree, STreeView<TSharedPtr< FClassViewerNode > >)
 		.SelectionMode(ESelectionMode::Single)
 		.TreeItemsSource(&RootTreeItems)
 		// Called to child items for any given parent item
@@ -1809,7 +1815,7 @@ void SClassViewer::OnClassViewerSelectionChanged( TSharedPtr<FClassViewerNode> I
 		// Check if the item passes the filter, parent items might be displayed but filtered out and thus not desired to be selected.
 		if ( ( Item->Class.IsValid() || !Class ))
 		{
-			if( Item->bPassesFilter )
+			if (Item->bPassesFilterRegardlessTextFilter)
 			{
 				OnClassPicked.ExecuteIfBound( Item->Class.Get() );
 			}
@@ -2531,6 +2537,7 @@ TSharedPtr<FClassViewerNode> SClassViewer::CreateNoneOption()
 
 	// The item "passes" the filter so it does not appear grayed out.
 	NoneItem->bPassesFilter = true;
+	NoneItem->bPassesFilterRegardlessTextFilter = true;
 
 	return NoneItem;
 }
diff --git a/Engine/Source/Editor/ClassViewer/Public/ClassViewerFilter.h b/Engine/Source/Editor/ClassViewer/Public/ClassViewerFilter.h
--- a/Engine/Source/Editor/ClassViewer/Public/ClassViewerFilter.h
+++ b/Engine/Source/Editor/ClassViewer/Public/ClassViewerFilter.h
@@ -42,10 +42,19 @@ class FClassViewerFilter : public IClassViewerFilter
 public:
 	FClassViewerFilter(const FClassViewerInitializationOptions& InInitOptions);
 
-	virtual bool IsNodeAllowed(const FClassViewerInitializationOptions& InInitOptions, const TSharedRef<FClassViewerNode>& Node);
+	/**
+	 * This function checks whether a node passes the filter defined by IsClassAllowed/IsUnloadedClassAllowed.
+	 *
+	 * @param InInitOptions				The Class Viewer/Picker options.
+	 * @param Node						The Node to check.
+	 * @param bCheckTextFilter			Whether to check the TextFilter. Disabling it could be useful e.g., to verify that the parent class of a IsNodeAllowed() object is also valid (regardless of the TextFilter, which will likely fail to pass).
+	 */
+	virtual bool IsNodeAllowed(const FClassViewerInitializationOptions& InInitOptions, const TSharedRef<FClassViewerNode>& Node, const bool bCheckTextFilter);
 
 	virtual bool IsClassAllowed(const FClassViewerInitializationOptions& InInitOptions, const UClass* InClass, TSharedRef< class FClassViewerFilterFuncs > InFilterFuncs ) override;
+	virtual bool IsClassAllowed(const FClassViewerInitializationOptions& InInitOptions, const UClass* InClass, TSharedRef< class FClassViewerFilterFuncs > InFilterFuncs, const bool bCheckTextFilter);
 	virtual bool IsUnloadedClassAllowed(const FClassViewerInitializationOptions& InInitOptions, const TSharedRef< const class IUnloadedBlueprintData > InUnloadedClassData, TSharedRef< class FClassViewerFilterFuncs > InFilterFuncs) override;
+	virtual bool IsUnloadedClassAllowed(const FClassViewerInitializationOptions& InInitOptions, const TSharedRef< const class IUnloadedBlueprintData > InUnloadedClassData, TSharedRef< class FClassViewerFilterFuncs > InFilterFuncs, const bool bCheckTextFilter);
 
 	TArray<UClass*> InternalClasses;
 	TArray<FDirectoryPath> InternalPaths;
