Fix to CL 7434472: Extended the new functionality of #jira UE-76564 to AddChildren_List.

diff --git a/Engine/Source/Editor/ClassViewer/Private/SClassViewer.cpp b/Engine/Source/Editor/ClassViewer/Private/SClassViewer.cpp
--- a/Engine/Source/Editor/ClassViewer/Private/SClassViewer.cpp
+++ b/Engine/Source/Editor/ClassViewer/Private/SClassViewer.cpp
@@ -382,6 +382,7 @@ namespace ClassViewer
 			{
 				TSharedPtr< FClassViewerNode > NewNode = MakeShareable(new FClassViewerNode(*InOriginalRootNode.Get()));
 				NewNode->bPassesFilter = true;
+				NewNode->bPassesFilterRegardlessTextFilter = true;
 				NewNode->PropertyHandle = InOriginalRootNode->PropertyHandle;
 
 				InOutNodeList.Add(NewNode);
@@ -417,6 +418,7 @@ namespace ClassViewer
 				{
 					TSharedPtr< FClassViewerNode > NewNode = MakeShareable(new FClassViewerNode(*ObjectClassRoot.Get()));
 					NewNode->bPassesFilter = true;
+					NewNode->bPassesFilterRegardlessTextFilter = true;
 					NewNode->PropertyHandle = InInitOptions.PropertyHandle;
 
 					InOutNodeList.Add(NewNode);
@@ -1812,10 +1814,12 @@ void SClassViewer::OnClassViewerSelectionChanged( TSharedPtr<FClassViewerNode> I
 			ClassViewer::Helpers::LoadClass( Item );
 		}
 
-		// Check if the item passes the filter, parent items might be displayed but filtered out and thus not desired to be selected.
+		// Check if the item passes the filter
 		if ( ( Item->Class.IsValid() || !Class ))
 		{
-			if (Item->bPassesFilterRegardlessTextFilter)
+			// Parent items might be displayed but filtered out by bPassesFilter, thus bPassesFilterRegardlessTextFilter makes sure to keep them selectable.
+			// In addition, Item->bPassesFilter would be redundant here as bPassesFilterRegardlessTextFilter = true if bPassesFilter = true
+			if (Item->bPassesFilterRegardlessTextFilter || Item->bPassesFilter)
 			{
 				OnClassPicked.ExecuteIfBound( Item->Class.Get() );
 			}
