Allow SClassViewer to specify an initially selected item

diff --git a/Engine/Source/Editor/ClassViewer/Private/SClassViewer.cpp b/Engine/Source/Editor/ClassViewer/Private/SClassViewer.cpp
--- a/Engine/Source/Editor/ClassViewer/Private/SClassViewer.cpp
+++ b/Engine/Source/Editor/ClassViewer/Private/SClassViewer.cpp
@@ -2365,10 +2365,12 @@ void SClassViewer::SetExpansionStatesInTree( TSharedPtr<FClassViewerNode> InItem
 int32 SClassViewer::CountTreeItems(FClassViewerNode* Node)
 {
 	if (Node == nullptr)
+	{
 		return 0;
+	}
 	int32 Count = 1;
 	TArray<TSharedPtr<FClassViewerNode>>& ChildArray = Node->GetChildrenList();
-	for (int i = 0; i < ChildArray.Num(); i++)
+	for (int32 i = 0; i < ChildArray.Num(); i++)
 	{
 		Count += CountTreeItems(ChildArray[i].Get());
 	}
@@ -2482,11 +2484,47 @@ void SClassViewer::Populate()
 		}
 
 		NumClasses = 0;
-		for (int i = 0; i < RootTreeItems.Num(); i++)
+		for (int32 i = 0; i < RootTreeItems.Num(); i++)
+		{
 			NumClasses += CountTreeItems(RootTreeItems[i].Get());
+		}
 
 		// Now that new items are in the tree, we need to request a refresh.
 		ClassTree->RequestTreeRefresh();
+
+		if (InitOptions.InitiallySelectedClass)
+		{
+			TArray<UClass*> ClassHierarchy;
+			UClass* CurrentClass = InitOptions.InitiallySelectedClass;
+
+			while (CurrentClass)
+			{
+				ClassHierarchy.Add(CurrentClass);
+				CurrentClass = CurrentClass->GetSuperClass();
+			}
+
+			ClassTree->SetItemExpansion(RootNode, true);
+
+			TSharedPtr<FClassViewerNode> ClassNode = RootNode;
+
+			for (int32 Index = ClassHierarchy.Num() - 2; Index >= 0; --Index)
+			{
+				for (const TSharedPtr<FClassViewerNode>& ChildClassNode : ClassNode->GetChildrenList())
+				{
+					UClass* ChildClass = ChildClassNode->Class.Get();
+					if (ChildClass == ClassHierarchy[Index])
+					{
+						ClassTree->SetItemExpansion(ChildClassNode, true);
+						ClassNode = ChildClassNode;
+						break;
+					}
+				}
+			}
+
+			ClassTree->SetSelection(ClassNode);
+
+			InitOptions.InitiallySelectedClass = nullptr;
+		}
 	}
 	else
 	{
@@ -2505,11 +2543,22 @@ void SClassViewer::Populate()
 		}
 
 		NumClasses = 0;
-		for (int i = 0; i < RootTreeItems.Num(); i++)
-			NumClasses += CountTreeItems(RootTreeItems[i].Get() );
+		for (int32 i = 0; i < RootTreeItems.Num(); i++)
+		{
+			NumClasses += CountTreeItems(RootTreeItems[i].Get());
+		}
 
 		// Now that new items are in the list, we need to request a refresh.
 		ClassList->RequestListRefresh();
+
+		if (InitOptions.InitiallySelectedClass)
+		{
+			if (TSharedPtr<FClassViewerNode> ClassNode = ClassViewer::Helpers::ClassHierarchy->FindNodeByClassName(ClassViewer::Helpers::ClassHierarchy->GetObjectRootNode(), InitOptions.InitiallySelectedClass->GetPathName()))
+			{
+				ClassList->SetSelection(ClassNode);
+			}
+			InitOptions.InitiallySelectedClass = nullptr;
+		}
 	}
 }
 
diff --git a/Engine/Source/Editor/ClassViewer/Public/ClassViewerModule.h b/Engine/Source/Editor/ClassViewer/Public/ClassViewerModule.h
--- a/Engine/Source/Editor/ClassViewer/Public/ClassViewerModule.h
+++ b/Engine/Source/Editor/ClassViewer/Public/ClassViewerModule.h
@@ -118,6 +118,9 @@ public:
 	/** false by default, restricts the class picker to only showing editor module classes */
 	bool bEditorClassesOnly;
 
+	/** Will set the initially selected row, if possible, to this class when the viewer is created */
+	UClass* InitiallySelectedClass;
+
 public:
 
 	/** Constructor */
@@ -135,6 +138,7 @@ public:
 		, NameTypeToDisplay(EClassViewerNameTypeToDisplay::ClassName)
 		, bAllowViewOptions(true)
 		, bEditorClassesOnly(false)
+		, InitiallySelectedClass(nullptr)
 	{
 	}
 };
