Fix linker error
#rnx

diff --git a/Engine/Source/Editor/ClassViewer/Public/SClassViewer.h b/Engine/Source/Editor/ClassViewer/Public/SClassViewer.h
--- a/Engine/Source/Editor/ClassViewer/Public/SClassViewer.h
+++ b/Engine/Source/Editor/ClassViewer/Public/SClassViewer.h
@@ -63,7 +63,7 @@ public:
 	virtual bool IsClassAllowed(const UClass* InClass) const;
 
 	/** Sends a requests to the Class Viewer to refresh itself the next chance it gets */
-	void Refresh();
+	CLASSVIEWER_API void Refresh();
 
 	/** Destroys the internal Class Hierarchy database */
 	static void DestroyClassHierarchy();
diff --git a/Engine/Source/Editor/KismetWidgets/KismetWidgets.Build.cs b/Engine/Source/Editor/KismetWidgets/KismetWidgets.Build.cs
--- a/Engine/Source/Editor/KismetWidgets/KismetWidgets.Build.cs
+++ b/Engine/Source/Editor/KismetWidgets/KismetWidgets.Build.cs
@@ -19,6 +19,7 @@ public class KismetWidgets : ModuleRules
                 "EditorStyle",
 				"UnrealEd",
 				"BlueprintGraph",
+				"ClassViewer"
 			}
 		);
 
