SClassViewer now maintains selection when filter changes

diff --git a/Engine/Source/Editor/ClassViewer/Private/SClassViewer.cpp b/Engine/Source/Editor/ClassViewer/Private/SClassViewer.cpp
--- a/Engine/Source/Editor/ClassViewer/Private/SClassViewer.cpp
+++ b/Engine/Source/Editor/ClassViewer/Private/SClassViewer.cpp
@@ -2379,6 +2379,21 @@ int32 SClassViewer::CountTreeItems(FClassViewerNode* Node)
 
 void SClassViewer::Populate()
 {
+	TArray<FName> PreviousSelection;
+	{
+		TArray<TSharedPtr<FClassViewerNode>> SelectedItems = GetSelectedItems();
+		if (SelectedItems.Num() > 0)
+		{
+			for (TSharedPtr<FClassViewerNode>& Node : SelectedItems)
+			{
+				if (Node.IsValid())
+				{
+					PreviousSelection.Add(Node->ClassPath);
+				}
+			}
+		}
+	}
+
 	bPendingSetExpansionStates = false;
 
 	// If showing a class tree, we may need to save expansion states.
@@ -2492,11 +2507,25 @@ void SClassViewer::Populate()
 		// Now that new items are in the tree, we need to request a refresh.
 		ClassTree->RequestTreeRefresh();
 
-		if (InitOptions.InitiallySelectedClass)
+		UClass* CurrentClass = nullptr;
+		if (PreviousSelection.Num() > 0)
 		{
-			TArray<UClass*> ClassHierarchy;
-			UClass* CurrentClass = InitOptions.InitiallySelectedClass;
+			if (TSharedPtr<FClassViewerNode> ClassNode = ClassViewer::Helpers::ClassHierarchy->FindNodeByClassName(ClassViewer::Helpers::ClassHierarchy->GetObjectRootNode(), PreviousSelection[0].ToString()))
+			{
+				if (ClassNode.IsValid())
+				{
+					CurrentClass = ClassNode->Class.Get();
+				}
+			}
+		}
+		else if (InitOptions.InitiallySelectedClass)
+		{
+			CurrentClass = InitOptions.InitiallySelectedClass;
+		}
 
+		if (CurrentClass)
+		{
+			TArray<UClass*> ClassHierarchy;
 			while (CurrentClass)
 			{
 				ClassHierarchy.Add(CurrentClass);
@@ -2551,9 +2580,19 @@ void SClassViewer::Populate()
 		// Now that new items are in the list, we need to request a refresh.
 		ClassList->RequestListRefresh();
 
-		if (InitOptions.InitiallySelectedClass)
+		FString ClassPathNameToSelect;
+		if (PreviousSelection.Num() > 0)
+		{
+			ClassPathNameToSelect = PreviousSelection[0].ToString();
+		}
+		else if (InitOptions.InitiallySelectedClass)
+		{
+			ClassPathNameToSelect = InitOptions.InitiallySelectedClass->GetPathName();
+		}
+
+		if (ClassPathNameToSelect.Len() > 0)
 		{
-			if (TSharedPtr<FClassViewerNode> ClassNode = ClassViewer::Helpers::ClassHierarchy->FindNodeByClassName(ClassViewer::Helpers::ClassHierarchy->GetObjectRootNode(), InitOptions.InitiallySelectedClass->GetPathName()))
+			if (TSharedPtr<FClassViewerNode> ClassNode = ClassViewer::Helpers::ClassHierarchy->FindNodeByClassName(ClassViewer::Helpers::ClassHierarchy->GetObjectRootNode(), ClassPathNameToSelect))
 			{
 				ClassList->SetSelection(ClassNode);
 			}
