Integrating CL 11291001 from //UE4/Dev-Editor

Class View no longer loses selection when losing focus.

diff --git a/Engine/Source/Editor/ClassViewer/Private/SClassViewer.cpp b/Engine/Source/Editor/ClassViewer/Private/SClassViewer.cpp
--- a/Engine/Source/Editor/ClassViewer/Private/SClassViewer.cpp
+++ b/Engine/Source/Editor/ClassViewer/Private/SClassViewer.cpp
@@ -2610,14 +2610,10 @@ FReply SClassViewer::OnKeyDown( const FGeometry& MyGeometry, const FKeyEvent& In
 
 FReply SClassViewer::OnFocusReceived( const FGeometry& MyGeometry, const FFocusEvent& InFocusEvent )
 {
-	if (RootTreeItems.Num() > 0)
+	if (InFocusEvent.GetCause() == EFocusCause::Navigation)
 	{
-		ClassTree->SetItemSelection(RootTreeItems[0], true, ESelectInfo::OnMouseClick);
-		ClassTree->SetItemExpansion(RootTreeItems[0], true);
-		OnClassViewerSelectionChanged(RootTreeItems[0],ESelectInfo::OnMouseClick);
+		FSlateApplication::Get().SetKeyboardFocus(SearchBox.ToSharedRef(), EFocusCause::SetDirectly);
 	}
-
-	FSlateApplication::Get().SetKeyboardFocus(SearchBox.ToSharedRef(), EFocusCause::SetDirectly);
 	
 	return FReply::Unhandled();
 }
