PR #6718: Scroll the Class Viewer to the selected item on refresh (Contributed by [REDACTED_PERSON])

diff --git a/Engine/Source/Editor/ClassViewer/Private/SClassViewer.cpp b/Engine/Source/Editor/ClassViewer/Private/SClassViewer.cpp
--- a/Engine/Source/Editor/ClassViewer/Private/SClassViewer.cpp
+++ b/Engine/Source/Editor/ClassViewer/Private/SClassViewer.cpp
@@ -2667,6 +2667,13 @@ void SClassViewer::Tick( const FGeometry& AllottedGeometry, const double InCurre
 		{
 			ExpandRootNodes();
 		}
+
+		// Scroll the first item into view if applicable
+		const TArray<TSharedPtr<FClassViewerNode>> SelectedItems = GetSelectedItems();
+		if (SelectedItems.Num() > 0)
+		{
+			ClassTree->RequestScrollIntoView(SelectedItems[0]);
+		}
 	}
 
 	if (bPendingSetExpansionStates)
