Correction to Class Viewer to correctly populate the class treee when a Parent Blueprint is loaded but not it's child. Reported on UDN.

diff --git a/Engine/Source/Editor/ClassViewer/Private/SClassViewer.cpp b/Engine/Source/Editor/ClassViewer/Private/SClassViewer.cpp
--- a/Engine/Source/Editor/ClassViewer/Private/SClassViewer.cpp
+++ b/Engine/Source/Editor/ClassViewer/Private/SClassViewer.cpp
@@ -326,7 +326,7 @@ namespace ClassViewer
 			TArray< TSharedPtr< FClassViewerNode > >& ChildList = InOriginalRootNode->GetChildrenList();
 			for(int32 ChildIdx = 0; ChildIdx < ChildList.Num(); ChildIdx++)
 			{
-				TSharedPtr< FClassViewerNode > NewNode = MakeShareable( new FClassViewerNode( *ChildList[ChildIdx].Get() ) );
+				TSharedPtr< FClassViewerNode > NewNode = MakeShared<FClassViewerNode>( *ChildList[ChildIdx].Get() );
 
 				const bool bChildrenPassesFilter = AddChildren_Tree(NewNode, ChildList[ChildIdx], InClassFilter, InInitOptions);
 				bReturnPassesFilter |= bChildrenPassesFilter;
@@ -358,13 +358,13 @@ namespace ClassViewer
 			const TSharedPtr< FClassViewerNode > ObjectClassRoot = ClassHierarchy->GetObjectRootNode();
 
 			// Duplicate the node, it will have no children.
-			InOutRootNode = MakeShareable(new FClassViewerNode(*ObjectClassRoot));
+			InOutRootNode = MakeShared<FClassViewerNode>(*ObjectClassRoot);
 
 			if (InInitOptions.bIsActorsOnly)
 			{
 				for (int32 ClassIdx = 0; ClassIdx < ObjectClassRoot->GetChildrenList().Num(); ClassIdx++)
 				{
-					TSharedPtr<FClassViewerNode> ChildNode = MakeShareable(new FClassViewerNode(*ObjectClassRoot->GetChildrenList()[ClassIdx].Get()));
+					TSharedPtr<FClassViewerNode> ChildNode = MakeShared<FClassViewerNode>(*ObjectClassRoot->GetChildrenList()[ClassIdx].Get());
 					if (AddChildren_Tree(ChildNode, ObjectClassRoot->GetChildrenList()[ClassIdx], InClassFilter, InInitOptions))
 					{
 						InOutRootNode->AddChild(ChildNode);
@@ -391,7 +391,7 @@ namespace ClassViewer
 			const bool bCheckTextFilter = true;
 			if (InClassFilter->IsNodeAllowed(InInitOptions, InOriginalRootNode.ToSharedRef(), bCheckTextFilter))
 			{
-				TSharedPtr< FClassViewerNode > NewNode = MakeShareable(new FClassViewerNode(*InOriginalRootNode.Get()));
+				TSharedPtr< FClassViewerNode > NewNode = MakeShared<FClassViewerNode>(*InOriginalRootNode.Get());
 				NewNode->bPassesFilter = true;
 				NewNode->bPassesFilterRegardlessTextFilter = true;
 				NewNode->PropertyHandle = InOriginalRootNode->PropertyHandle;
@@ -427,7 +427,7 @@ namespace ClassViewer
 				const bool bCheckTextFilter = true;
 				if (InClassFilter->IsNodeAllowed(InInitOptions, ObjectClassRoot.ToSharedRef(), bCheckTextFilter))
 				{
-					TSharedPtr< FClassViewerNode > NewNode = MakeShareable(new FClassViewerNode(*ObjectClassRoot.Get()));
+					TSharedPtr< FClassViewerNode > NewNode = MakeShared<FClassViewerNode>(*ObjectClassRoot.Get());
 					NewNode->bPassesFilter = true;
 					NewNode->bPassesFilterRegardlessTextFilter = true;
 					NewNode->PropertyHandle = InInitOptions.PropertyHandle;
@@ -1068,17 +1068,13 @@ void FClassHierarchy::AddChildren_NoFilter( TSharedPtr< FClassViewerNode >& InOu
 {
 	UClass* RootClass = UObject::StaticClass();
 
-	ObjectClassRoot = MakeShareable(new FClassViewerNode(RootClass));
+	ObjectClassRoot = MakeShared<FClassViewerNode>(RootClass);
 	ObjectClassRoot->Class = RootClass;
 
 	InOutClassPathToNode.Add(ObjectClassRoot->ClassPath, ObjectClassRoot);
 
-	TMap< UClass*, TSharedPtr< FClassViewerNode > > Nodes;
-
-	Nodes.Add(RootClass, ObjectClassRoot);
-
-	TSet<UClass*> Visited;
-	Visited.Add(RootClass);
+	TMap< UClass*, TSharedPtr< FClassViewerNode > > VisitedNodes;
+	VisitedNodes.Add(RootClass, ObjectClassRoot);
 
 	// Go through all of the classes children and see if they should be added to the list.
 	for ( TObjectIterator<UClass> ClassIt; ClassIt; ++ClassIt )
@@ -1092,33 +1088,56 @@ void FClassHierarchy::AddChildren_NoFilter( TSharedPtr< FClassViewerNode >& InOu
 			continue;
 		}
 		
-		TSharedPtr<FClassViewerNode>& Entry = Nodes.FindOrAdd(CurrentClass);
-		if ( Visited.Contains(CurrentClass) )
+		TSharedPtr<FClassViewerNode>& Entry = VisitedNodes.FindOrAdd(CurrentClass);
+		// If the entry is valid we already visited it.
+		if (Entry.IsValid())
 		{
 			continue;
 		}
 		else
 		{
+			// Process this node and all it's parent classes.
 			while ( CurrentClass->GetSuperClass() != nullptr )
 			{
-				TSharedPtr<FClassViewerNode>& ParentEntry = Nodes.FindOrAdd(CurrentClass->GetSuperClass());
+				TSharedPtr<FClassViewerNode>& ParentEntry = VisitedNodes.FindOrAdd(CurrentClass->GetSuperClass());
+				// If ParentEntry is not valid, it's the first time we visited that class.
+				// We add it to the ClassPathToNode map if it's not there
 				if ( !ParentEntry.IsValid() )
 				{
-					ParentEntry = MakeShareable(new FClassViewerNode(CurrentClass->GetSuperClass()));
-                    InOutClassPathToNode.Add(ParentEntry->ClassPath, ParentEntry);
-				}
+					/// If the class is already present, make sure we use it.
+					FName ParentClassPath = FName(*CurrentClass->GetSuperClass()->GetPathName());
+					TSharedPtr<FClassViewerNode>* AlreadyExisting = InOutClassPathToNode.Find(ParentClassPath);
+					if (AlreadyExisting)
+					{
+						ParentEntry = *AlreadyExisting;
+					}
+					else
+					{
+						ParentEntry = MakeShared<FClassViewerNode>(CurrentClass->GetSuperClass());
+                    	InOutClassPathToNode.Add(ParentEntry->ClassPath, ParentEntry);
+					}
 
-				TSharedPtr<FClassViewerNode>& MyEntry = Nodes.FindOrAdd(CurrentClass);
-				if ( !MyEntry.IsValid() )
-				{
-					MyEntry = MakeShareable(new FClassViewerNode(CurrentClass));
-                    InOutClassPathToNode.Add(MyEntry->ClassPath, MyEntry);
 				}
 
-				if ( !Visited.Contains(CurrentClass) )
+				TSharedPtr<FClassViewerNode>& MyEntry = VisitedNodes.FindOrAdd(CurrentClass);
+				// If MyEntry is not valid, it's the first time we visited that class.
+				// We add it to the ClassPathToNode map if it's not there and add it as a child to MyEntry
+				if ( !MyEntry.IsValid() )
 				{
-					ParentEntry->AddChild(MyEntry);
-					Visited.Add(CurrentClass);
+					/// If the class is already present, make sure we use it.
+					FName ClassPath = FName(*CurrentClass->GetPathName());
+					TSharedPtr<FClassViewerNode>* AlreadyExisting = InOutClassPathToNode.Find(ClassPath);
+					if (AlreadyExisting)
+					{
+						MyEntry = *AlreadyExisting;
+						ParentEntry->AddUniqueChild(MyEntry);
+					}
+					else
+					{
+						MyEntry = MakeShared<FClassViewerNode>(CurrentClass);
+						InOutClassPathToNode.Add(MyEntry->ClassPath, MyEntry);
+						ParentEntry->AddChild(MyEntry);
+					}
 				}
 
 				CurrentClass = CurrentClass->GetSuperClass();
@@ -1353,7 +1372,7 @@ void FClassHierarchy::LoadUnloadedTagData(TSharedPtr<FClassViewerNode>& InOutCla
 		ClassDisplayName = ClassName;
 	}
 	// Create the viewer node. We use the name without _C for both
-	InOutClassViewerNode = MakeShareable(new FClassViewerNode(ClassName, ClassDisplayName));
+	InOutClassViewerNode = MakeShared<FClassViewerNode>(ClassName, ClassDisplayName);
 
 	InOutClassViewerNode->BlueprintAssetPath = InAssetData.ObjectPath;
 
@@ -2630,7 +2649,7 @@ void SClassViewer::DestroyClassHierarchy()
 
 TSharedPtr<FClassViewerNode> SClassViewer::CreateNoneOption()
 {
-	TSharedPtr<FClassViewerNode> NoneItem = MakeShareable( new FClassViewerNode("None", "None") );
+	TSharedPtr<FClassViewerNode> NoneItem = MakeShared<FClassViewerNode>("None", "None");
 
 	// The item "passes" the filter so it does not appear grayed out.
 	NoneItem->bPassesFilter = true;
