Extends the class viewer module to support multiple custom class filters along with an optional associated view option flag.

Additional changes:
- Deprecates the previous method for specifying a singular custom class viewer filter and updates all existing occurrences of this pattern in engine code.
- Extends the property editor utilities interface to expose custom class filter(s) that can be applied to the class picker widget used for editing class property values.
- Adds an implementation of this interface to SDetailsView such that additional class filter(s) can now be configured to be applied to all underlying class property nodes.

diff --git a/Engine/Plugins/Animation/LiveLink/Source/LiveLinkEditor/Private/LiveLinkClientPanelToolbar.cpp b/Engine/Plugins/Animation/LiveLink/Source/LiveLinkEditor/Private/LiveLinkClientPanelToolbar.cpp
--- a/Engine/Plugins/Animation/LiveLink/Source/LiveLinkEditor/Private/LiveLinkClientPanelToolbar.cpp
+++ b/Engine/Plugins/Animation/LiveLink/Source/LiveLinkEditor/Private/LiveLinkClientPanelToolbar.cpp
@@ -296,7 +296,7 @@ private:
 		FClassViewerInitializationOptions Options;
 		Options.Mode = EClassViewerMode::ClassPicker;
 
-		Options.ClassFilter = MakeShared<FLiveLinkRoleClassFilter>();
+		Options.ClassFilters.Add(MakeShared<FLiveLinkRoleClassFilter>());
 
 		RoleClassPicker->ClearChildren();
 		RoleClassPicker->AddSlot()
diff --git a/Engine/Plugins/Animation/LiveLink/Source/LiveLinkEditor/Private/VirtualSubjects/LiveLinkBlueprintVirtualSubjectFactory.cpp b/Engine/Plugins/Animation/LiveLink/Source/LiveLinkEditor/Private/VirtualSubjects/LiveLinkBlueprintVirtualSubjectFactory.cpp
--- a/Engine/Plugins/Animation/LiveLink/Source/LiveLinkEditor/Private/VirtualSubjects/LiveLinkBlueprintVirtualSubjectFactory.cpp
+++ b/Engine/Plugins/Animation/LiveLink/Source/LiveLinkEditor/Private/VirtualSubjects/LiveLinkBlueprintVirtualSubjectFactory.cpp
@@ -122,7 +122,7 @@ bool ULiveLinkBlueprintVirtualSubjectFactory::ConfigureProperties()
 			// Fill in options
 			FClassViewerInitializationOptions Options;
 			Options.Mode = EClassViewerMode::ClassPicker;
-			Options.ClassFilter = MakeShared<FLiveLinkRoleClassFilter>();
+			Options.ClassFilters.Add(MakeShared<FLiveLinkRoleClassFilter>());
 
 			return
 				SNew(SBox)
@@ -151,7 +151,7 @@ bool ULiveLinkBlueprintVirtualSubjectFactory::ConfigureProperties()
 			// Fill in options
 			FClassViewerInitializationOptions Options;
 			Options.Mode = EClassViewerMode::ClassPicker;
-			Options.ClassFilter = MakeShared<FLiveLinkRoleClassFilter>();
+			Options.ClassFilters.Add(MakeShared<FLiveLinkRoleClassFilter>());
 
 			PickerWindow = SNew(SWindow)
 				.Title(LOCTEXT("PickerTitle", "Select Role"))
diff --git a/Engine/Plugins/Compositing/Composure/Source/ComposureLayersEditor/Private/CompElementCollectionViewModel.cpp b/Engine/Plugins/Compositing/Composure/Source/ComposureLayersEditor/Private/CompElementCollectionViewModel.cpp
--- a/Engine/Plugins/Compositing/Composure/Source/ComposureLayersEditor/Private/CompElementCollectionViewModel.cpp
+++ b/Engine/Plugins/Compositing/Composure/Source/ComposureLayersEditor/Private/CompElementCollectionViewModel.cpp
@@ -77,7 +77,7 @@ static UClass* CompElementCollectionViewModel_Impl::PromptForElementClass(const
 	FClassViewerInitializationOptions Options;
 	Options.Mode = EClassViewerMode::ClassPicker;
 	TSharedPtr<FCompElementClassFilter> Filter = MakeShareable(new FCompElementClassFilter);
-	Options.ClassFilter = Filter;
+	Options.ClassFilters.Add(Filter.ToSharedRef());
 
 	Options.ExtraPickerCommonClasses.Reserve(ChoiceClasses.Num());
 	for (TSubclassOf<ACompositingElement> Class : ChoiceClasses)
diff --git a/Engine/Plugins/Developer/TraceSourceFiltering/Source/SourceFilteringEditor/Private/EditorSessionSourceFilterService.cpp b/Engine/Plugins/Developer/TraceSourceFiltering/Source/SourceFilteringEditor/Private/EditorSessionSourceFilterService.cpp
--- a/Engine/Plugins/Developer/TraceSourceFiltering/Source/SourceFilteringEditor/Private/EditorSessionSourceFilterService.cpp
+++ b/Engine/Plugins/Developer/TraceSourceFiltering/Source/SourceFilteringEditor/Private/EditorSessionSourceFilterService.cpp
@@ -179,7 +179,7 @@ TSharedRef<SWidget> FEditorSessionSourceFilterService::GetFilterPickerWidget(FOn
 	Options.bShowUnloadedBlueprints = true;
 	Options.bShowNoneOption = false;
 	TSharedPtr<FFilterClassFilter> ClassFilter = MakeShareable(new FFilterClassFilter);
-	Options.ClassFilter = ClassFilter;
+	Options.ClassFilters.Add(ClassFilter.ToSharedRef());
 
 	FOnClassPicked ClassPicked = FOnClassPicked::CreateLambda([InFilterClassPicked](UClass* Class)
 	{
@@ -242,7 +242,7 @@ TSharedRef<SWidget> FEditorSessionSourceFilterService::GetClassFilterPickerWidge
 	GetClassFilters(ExistingClasses);
 
 	TSharedPtr<FFilterClassFilter> ClassFilter = MakeShared<FFilterClassFilter>(ExistingClasses);
-	Options.ClassFilter = ClassFilter;
+	Options.ClassFilters.Add(ClassFilter.ToSharedRef());
 
 	FOnClassPicked ClassPicked = FOnClassPicked::CreateLambda([InFilterClassPicked](UClass* Class)
 	{
diff --git a/Engine/Plugins/Experimental/CommonConversation/Source/CommonConversationEditor/Private/ConversationEditor.cpp b/Engine/Plugins/Experimental/CommonConversation/Source/CommonConversationEditor/Private/ConversationEditor.cpp
--- a/Engine/Plugins/Experimental/CommonConversation/Source/CommonConversationEditor/Private/ConversationEditor.cpp
+++ b/Engine/Plugins/Experimental/CommonConversation/Source/CommonConversationEditor/Private/ConversationEditor.cpp
@@ -925,7 +925,7 @@ TSharedRef<SWidget> FConversationEditor::HandleCreateNewClassMenu(UClass* BaseCl
 {
 	FClassViewerInitializationOptions Options;
 	Options.bShowUnloadedBlueprints = true;
-	Options.ClassFilter = MakeShareable( new FNewNodeClassFilter(BaseClass) );
+	Options.ClassFilters.Add(MakeShareable( new FNewNodeClassFilter(BaseClass) ));
 
 	FOnClassPicked OnPicked( FOnClassPicked::CreateSP( this, &FConversationEditor::HandleNewNodeClassPicked ) );
 
diff --git a/Engine/Plugins/Experimental/ControlRig/Source/ControlRigEditor/Private/ControlRigBlueprintFactory.cpp b/Engine/Plugins/Experimental/ControlRig/Source/ControlRigEditor/Private/ControlRigBlueprintFactory.cpp
--- a/Engine/Plugins/Experimental/ControlRig/Source/ControlRigEditor/Private/ControlRigBlueprintFactory.cpp
+++ b/Engine/Plugins/Experimental/ControlRig/Source/ControlRigEditor/Private/ControlRigBlueprintFactory.cpp
@@ -176,7 +176,7 @@ private:
 		Options.Mode = EClassViewerMode::ClassPicker;
 
 		TSharedPtr<FControlRigBlueprintParentFilter> Filter = MakeShareable(new FControlRigBlueprintParentFilter());
-		Options.ClassFilter = Filter;
+		Options.ClassFilters.Add(Filter.ToSharedRef());
 
 		// All child child classes of UControlRig are valid.
 		Filter->AllowedChildrenOfClasses.Add(UControlRig::StaticClass());
diff --git a/Engine/Plugins/Experimental/ControlRig/Source/ControlRigEditor/Private/ControlRigEditorModule.cpp b/Engine/Plugins/Experimental/ControlRig/Source/ControlRigEditor/Private/ControlRigEditorModule.cpp
--- a/Engine/Plugins/Experimental/ControlRig/Source/ControlRigEditor/Private/ControlRigEditorModule.cpp
+++ b/Engine/Plugins/Experimental/ControlRig/Source/ControlRigEditor/Private/ControlRigEditorModule.cpp
@@ -658,7 +658,7 @@ TSharedRef< SWidget > FControlRigEditorModule::GenerateAnimationMenu(TWeakPtr<IA
 								Options.NameTypeToDisplay = EClassViewerNameTypeToDisplay::DisplayName;
 
 								TSharedPtr<FControlRigClassFilter> ClassFilter = MakeShareable(new FControlRigClassFilter(bFilterAssetBySkeleton, true, true, Skeleton));
-								Options.ClassFilter = ClassFilter;
+								Options.ClassFilters.Add(ClassFilter.ToSharedRef());
 								Options.bShowNoneOption = false;
 
 								FClassViewerModule& ClassViewerModule = FModuleManager::LoadModuleChecked<FClassViewerModule>("ClassViewer");
diff --git a/Engine/Plugins/Experimental/ControlRig/Source/ControlRigEditor/Private/Sequencer/ControlRigParameterTrackEditor.cpp b/Engine/Plugins/Experimental/ControlRig/Source/ControlRigEditor/Private/Sequencer/ControlRigParameterTrackEditor.cpp
--- a/Engine/Plugins/Experimental/ControlRig/Source/ControlRigEditor/Private/Sequencer/ControlRigParameterTrackEditor.cpp
+++ b/Engine/Plugins/Experimental/ControlRig/Source/ControlRigEditor/Private/Sequencer/ControlRigParameterTrackEditor.cpp
@@ -621,7 +621,7 @@ void FControlRigParameterTrackEditor::BakeToControlRigSubMenu(FMenuBuilder& Menu
 		Options.bShowUnloadedBlueprints = true;
 		Options.NameTypeToDisplay = EClassViewerNameTypeToDisplay::DisplayName;
 		TSharedPtr<FControlRigClassFilter> ClassFilter = MakeShareable(new FControlRigClassFilter(bFilterAssetBySkeleton,true, true, Skeleton));
-		Options.ClassFilter = ClassFilter;
+		Options.ClassFilters.Add(ClassFilter.ToSharedRef());
 		Options.bShowNoneOption = false;
 
 		FClassViewerModule& ClassViewerModule = FModuleManager::LoadModuleChecked<FClassViewerModule>("ClassViewer");
@@ -1157,7 +1157,7 @@ void FControlRigParameterTrackEditor::HandleAddControlRigSubMenu(FMenuBuilder& M
 		Options.NameTypeToDisplay = EClassViewerNameTypeToDisplay::DisplayName;
 
 		TSharedPtr<FControlRigClassFilter> ClassFilter = MakeShareable(new FControlRigClassFilter(bFilterAssetBySkeleton, bFilterAssetByAnimatableControls, false, Skeleton));
-		Options.ClassFilter = ClassFilter;
+		Options.ClassFilters.Add(ClassFilter.ToSharedRef());
 		Options.bShowNoneOption = false;
 
 		FClassViewerModule& ClassViewerModule = FModuleManager::LoadModuleChecked<FClassViewerModule>("ClassViewer");
diff --git a/Engine/Plugins/Media/MediaFrameworkUtilities/Source/MediaFrameworkUtilitiesEditor/Private/UI/MediaProfileMenuEntry.cpp b/Engine/Plugins/Media/MediaFrameworkUtilities/Source/MediaFrameworkUtilitiesEditor/Private/UI/MediaProfileMenuEntry.cpp
--- a/Engine/Plugins/Media/MediaFrameworkUtilities/Source/MediaFrameworkUtilitiesEditor/Private/UI/MediaProfileMenuEntry.cpp
+++ b/Engine/Plugins/Media/MediaFrameworkUtilities/Source/MediaFrameworkUtilitiesEditor/Private/UI/MediaProfileMenuEntry.cpp
@@ -86,7 +86,7 @@ struct FMediaProfileMenuEntryImpl
 
 		const FText TitleText = LOCTEXT("CreateMediaProfileOptions", "Pick Media Profile Class");
 		FClassViewerInitializationOptions Options;
-		Options.ClassFilter = MakeShared<FModifierClassFilter>();
+		Options.ClassFilters.Add(MakeShared<FModifierClassFilter>());
 		UClass* ChosenClass = nullptr;
 		const bool bPressedOk = SClassPickerDialog::PickClass(TitleText, Options, ChosenClass, UMediaProfile::StaticClass());
 
diff --git a/Engine/Plugins/Runtime/AudioModulation/Source/AudioModulationEditor/Private/Factories/SoundModulationGeneratorFactory.cpp b/Engine/Plugins/Runtime/AudioModulation/Source/AudioModulationEditor/Private/Factories/SoundModulationGeneratorFactory.cpp
--- a/Engine/Plugins/Runtime/AudioModulation/Source/AudioModulationEditor/Private/Factories/SoundModulationGeneratorFactory.cpp
+++ b/Engine/Plugins/Runtime/AudioModulation/Source/AudioModulationEditor/Private/Factories/SoundModulationGeneratorFactory.cpp
@@ -59,7 +59,7 @@ bool USoundModulationGeneratorFactory::ConfigureProperties()
 	FClassViewerModule& ClassViewerModule = FModuleManager::LoadModuleChecked<FClassViewerModule>("ClassViewer");
 
 	FClassViewerInitializationOptions Options;
-	Options.ClassFilter = MakeShared<FGeneratorClassViewerFilter>();
+	Options.ClassFilters.Add(MakeShared<FGeneratorClassViewerFilter>());
 	Options.Mode = EClassViewerMode::ClassPicker;
 
 	const FText TitleText = LOCTEXT("CreateModulationGeneratorClassSelect", "Select Generator Class");
diff --git a/Engine/Plugins/Runtime/AudioModulation/Source/AudioModulationEditor/Private/Factories/SoundModulationParameterFactory.cpp b/Engine/Plugins/Runtime/AudioModulation/Source/AudioModulationEditor/Private/Factories/SoundModulationParameterFactory.cpp
--- a/Engine/Plugins/Runtime/AudioModulation/Source/AudioModulationEditor/Private/Factories/SoundModulationParameterFactory.cpp
+++ b/Engine/Plugins/Runtime/AudioModulation/Source/AudioModulationEditor/Private/Factories/SoundModulationParameterFactory.cpp
@@ -48,7 +48,7 @@ bool USoundModulationParameterFactory::ConfigureProperties()
 	Options.Mode = EClassViewerMode::ClassPicker;
 
 	TSharedPtr<FAssetClassParentFilter> Filter = MakeShared<FAssetClassParentFilter>();
-	Options.ClassFilter = Filter;
+	Options.ClassFilters.Add(Filter.ToSharedRef());
 
 	Filter->DisallowedClassFlags = CLASS_Abstract | CLASS_Deprecated | CLASS_NewerVersionExists;
 	Filter->AllowedChildrenOfClasses.Add(USoundModulationParameter::StaticClass());
diff --git a/Engine/Plugins/Runtime/AudioSynesthesia/Source/AudioSynesthesiaEditor/Private/AudioSynesthesiaNRTFactory.cpp b/Engine/Plugins/Runtime/AudioSynesthesia/Source/AudioSynesthesiaEditor/Private/AudioSynesthesiaNRTFactory.cpp
--- a/Engine/Plugins/Runtime/AudioSynesthesia/Source/AudioSynesthesiaEditor/Private/AudioSynesthesiaNRTFactory.cpp
+++ b/Engine/Plugins/Runtime/AudioSynesthesia/Source/AudioSynesthesiaEditor/Private/AudioSynesthesiaNRTFactory.cpp
@@ -32,7 +32,7 @@ bool UAudioSynesthesiaNRTFactory::ConfigureProperties()
 	Options.Mode = EClassViewerMode::ClassPicker;
 
 	TSharedPtr<Audio::FAssetClassParentFilter> Filter = MakeShareable(new Audio::FAssetClassParentFilter);
-	Options.ClassFilter = Filter;
+	Options.ClassFilters.Add(Filter.ToSharedRef());
 
 	Filter->DisallowedClassFlags = CLASS_Abstract | CLASS_Deprecated | CLASS_NewerVersionExists;
 	Filter->AllowedChildrenOfClasses.Add(UAudioSynesthesiaNRT::StaticClass());
diff --git a/Engine/Plugins/Runtime/AudioSynesthesia/Source/AudioSynesthesiaEditor/Private/AudioSynesthesiaNRTSettingsFactory.cpp b/Engine/Plugins/Runtime/AudioSynesthesia/Source/AudioSynesthesiaEditor/Private/AudioSynesthesiaNRTSettingsFactory.cpp
--- a/Engine/Plugins/Runtime/AudioSynesthesia/Source/AudioSynesthesiaEditor/Private/AudioSynesthesiaNRTSettingsFactory.cpp
+++ b/Engine/Plugins/Runtime/AudioSynesthesia/Source/AudioSynesthesiaEditor/Private/AudioSynesthesiaNRTSettingsFactory.cpp
@@ -31,7 +31,7 @@ bool UAudioSynesthesiaNRTSettingsFactory::ConfigureProperties()
 	Options.Mode = EClassViewerMode::ClassPicker;
 
 	TSharedPtr<Audio::FAssetClassParentFilter> Filter = MakeShareable(new Audio::FAssetClassParentFilter);
-	Options.ClassFilter = Filter;
+	Options.ClassFilters.Add(Filter.ToSharedRef());
 
 	Filter->DisallowedClassFlags = CLASS_Abstract | CLASS_Deprecated | CLASS_NewerVersionExists;
 	Filter->AllowedChildrenOfClasses.Add(UAudioSynesthesiaNRTSettings::StaticClass());
diff --git a/Engine/Plugins/Runtime/AudioSynesthesia/Source/AudioSynesthesiaEditor/Private/AudioSynesthesiaSettingsFactory.cpp b/Engine/Plugins/Runtime/AudioSynesthesia/Source/AudioSynesthesiaEditor/Private/AudioSynesthesiaSettingsFactory.cpp
--- a/Engine/Plugins/Runtime/AudioSynesthesia/Source/AudioSynesthesiaEditor/Private/AudioSynesthesiaSettingsFactory.cpp
+++ b/Engine/Plugins/Runtime/AudioSynesthesia/Source/AudioSynesthesiaEditor/Private/AudioSynesthesiaSettingsFactory.cpp
@@ -31,7 +31,7 @@ bool UAudioSynesthesiaSettingsFactory::ConfigureProperties()
 	Options.Mode = EClassViewerMode::ClassPicker;
 
 	TSharedPtr<Audio::FAssetClassParentFilter> Filter = MakeShareable(new Audio::FAssetClassParentFilter);
-	Options.ClassFilter = Filter;
+	Options.ClassFilters.Add(Filter.ToSharedRef());
 
 	Filter->DisallowedClassFlags = CLASS_Abstract | CLASS_Deprecated | CLASS_NewerVersionExists;
 	Filter->AllowedChildrenOfClasses.Add(UAudioSynesthesiaSettings::StaticClass());
diff --git a/Engine/Plugins/Runtime/DataRegistry/Source/DataRegistryEditor/Private/DataRegistryFactory.cpp b/Engine/Plugins/Runtime/DataRegistry/Source/DataRegistryEditor/Private/DataRegistryFactory.cpp
--- a/Engine/Plugins/Runtime/DataRegistry/Source/DataRegistryEditor/Private/DataRegistryFactory.cpp
+++ b/Engine/Plugins/Runtime/DataRegistry/Source/DataRegistryEditor/Private/DataRegistryFactory.cpp
@@ -62,7 +62,7 @@ bool UDataRegistryFactory::ConfigureProperties()
 	// Options.InitiallySelectedClass = UDataRegistry::StaticClass();
 
 	TSharedPtr<FAssetClassParentFilter> Filter = MakeShareable(new FAssetClassParentFilter);
-	Options.ClassFilter = Filter;
+	Options.ClassFilters.Add(Filter.ToSharedRef());
 
 	Filter->DisallowedClassFlags = CLASS_Abstract | CLASS_Deprecated | CLASS_NewerVersionExists | CLASS_HideDropDown;
 	Filter->AllowedChildrenOfClasses.Add(UDataRegistry::StaticClass());
diff --git a/Engine/Plugins/Runtime/GameplayAbilities/Source/GameplayAbilitiesEditor/Private/GameplayAbilitiesBlueprintFactory.cpp b/Engine/Plugins/Runtime/GameplayAbilities/Source/GameplayAbilitiesEditor/Private/GameplayAbilitiesBlueprintFactory.cpp
--- a/Engine/Plugins/Runtime/GameplayAbilities/Source/GameplayAbilitiesEditor/Private/GameplayAbilitiesBlueprintFactory.cpp
+++ b/Engine/Plugins/Runtime/GameplayAbilities/Source/GameplayAbilitiesEditor/Private/GameplayAbilitiesBlueprintFactory.cpp
@@ -174,7 +174,7 @@ private:
 
 		// All child child classes of UGameplayAbility are valid.
 		Filter->AllowedChildrenOfClasses.Add(UGameplayAbility::StaticClass());
-		Options.ClassFilter = Filter;
+		Options.ClassFilters.Add(Filter.ToSharedRef());
 
 		ParentClassContainer->ClearChildren();
 		ParentClassContainer->AddSlot()
diff --git a/Engine/Plugins/Runtime/SoundCueTemplates/Source/SoundCueTemplatesEditor/Private/SoundCueTemplateFactory.cpp b/Engine/Plugins/Runtime/SoundCueTemplates/Source/SoundCueTemplatesEditor/Private/SoundCueTemplateFactory.cpp
--- a/Engine/Plugins/Runtime/SoundCueTemplates/Source/SoundCueTemplatesEditor/Private/SoundCueTemplateFactory.cpp
+++ b/Engine/Plugins/Runtime/SoundCueTemplates/Source/SoundCueTemplatesEditor/Private/SoundCueTemplateFactory.cpp
@@ -61,7 +61,7 @@ bool USoundCueTemplateFactory::ConfigureProperties()
 
 	FClassViewerInitializationOptions Options;
 	Options.Mode = EClassViewerMode::ClassPicker;
-	Options.ClassFilter = Filter;
+	Options.ClassFilters.Add(Filter.ToSharedRef());
 
 	const FText TitleText = LOCTEXT("CreateSoundCueTemplateOfType", "Pick Type of SoundCueTemplate");
 	UClass* ChosenClass = nullptr;
diff --git a/Engine/Plugins/VirtualProduction/RemoteControl/Source/RemoteControlUI/Private/UI/SRemoteControlPanel.cpp b/Engine/Plugins/VirtualProduction/RemoteControl/Source/RemoteControlUI/Private/UI/SRemoteControlPanel.cpp
--- a/Engine/Plugins/VirtualProduction/RemoteControl/Source/RemoteControlUI/Private/UI/SRemoteControlPanel.cpp
+++ b/Engine/Plugins/VirtualProduction/RemoteControl/Source/RemoteControlUI/Private/UI/SRemoteControlPanel.cpp
@@ -584,7 +584,7 @@ TSharedRef<SWidget> SRemoteControlPanel::CreateExposeByClassWidget()
 	
 	FClassViewerInitializationOptions Options;
 	{
-		Options.ClassFilter = Filter;
+		Options.ClassFilters.Add(Filter.ToSharedRef());
 		Options.bIsPlaceableOnly = true;
 		Options.Mode = EClassViewerMode::ClassPicker;
 		Options.DisplayMode = EClassViewerDisplayMode::ListView;
diff --git a/Engine/Source/Editor/AnimationModifiers/Private/AnimationModifierHelpers.h b/Engine/Source/Editor/AnimationModifiers/Private/AnimationModifierHelpers.h
--- a/Engine/Source/Editor/AnimationModifiers/Private/AnimationModifierHelpers.h
+++ b/Engine/Source/Editor/AnimationModifiers/Private/AnimationModifierHelpers.h
@@ -34,8 +34,8 @@ public:
 		FClassViewerInitializationOptions Options;
 		Options.bShowUnloadedBlueprints = true;
 		Options.bShowNoneOption = false;
-		TSharedPtr<FModifierClassFilter> ClassFilter = MakeShareable(new FModifierClassFilter);
-		Options.ClassFilter = ClassFilter;
+		TSharedRef<FModifierClassFilter> ClassFilter = MakeShared<FModifierClassFilter>();
+		Options.ClassFilters.Add(ClassFilter);
 
 		return SNew(SBox)
 		.WidthOverride(280)
diff --git a/Engine/Source/Editor/AudioEditor/Private/Factories/SoundSourceEffectFactory.cpp b/Engine/Source/Editor/AudioEditor/Private/Factories/SoundSourceEffectFactory.cpp
--- a/Engine/Source/Editor/AudioEditor/Private/Factories/SoundSourceEffectFactory.cpp
+++ b/Engine/Source/Editor/AudioEditor/Private/Factories/SoundSourceEffectFactory.cpp
@@ -33,7 +33,7 @@ bool USoundSourceEffectFactory::ConfigureProperties()
 	Options.Mode = EClassViewerMode::ClassPicker;
 
 	TSharedPtr<FAssetClassParentFilter> Filter = MakeShareable(new FAssetClassParentFilter);
-	Options.ClassFilter = Filter;
+	Options.ClassFilters.Add(Filter.ToSharedRef());
 
 	Filter->DisallowedClassFlags = CLASS_Abstract | CLASS_Deprecated | CLASS_NewerVersionExists;
 	Filter->AllowedChildrenOfClasses.Add(USoundEffectSourcePreset::StaticClass());
diff --git a/Engine/Source/Editor/AudioEditor/Private/Factories/SoundSubmixEffectFactory.cpp b/Engine/Source/Editor/AudioEditor/Private/Factories/SoundSubmixEffectFactory.cpp
--- a/Engine/Source/Editor/AudioEditor/Private/Factories/SoundSubmixEffectFactory.cpp
+++ b/Engine/Source/Editor/AudioEditor/Private/Factories/SoundSubmixEffectFactory.cpp
@@ -33,7 +33,7 @@ bool USoundSubmixEffectFactory::ConfigureProperties()
 	Options.Mode = EClassViewerMode::ClassPicker;
 
 	TSharedPtr<FAssetClassParentFilter> Filter = MakeShareable(new FAssetClassParentFilter);
-	Options.ClassFilter = Filter;
+	Options.ClassFilters.Add(Filter.ToSharedRef());
 
 	Filter->DisallowedClassFlags = CLASS_Abstract | CLASS_Deprecated | CLASS_NewerVersionExists;
 	Filter->AllowedChildrenOfClasses.Add(USoundEffectSubmixPreset::StaticClass());
diff --git a/Engine/Source/Editor/BehaviorTreeEditor/Private/BehaviorTreeEditor.cpp b/Engine/Source/Editor/BehaviorTreeEditor/Private/BehaviorTreeEditor.cpp
--- a/Engine/Source/Editor/BehaviorTreeEditor/Private/BehaviorTreeEditor.cpp
+++ b/Engine/Source/Editor/BehaviorTreeEditor/Private/BehaviorTreeEditor.cpp
@@ -1694,7 +1694,7 @@ TSharedRef<SWidget> FBehaviorTreeEditor::HandleCreateNewTaskMenu() const
 {
 	FClassViewerInitializationOptions Options;
 	Options.bShowUnloadedBlueprints = true;
-	Options.ClassFilter = MakeShareable( new FBehaviorTreeEditorUtils::FNewNodeClassFilter<UBTTask_BlueprintBase> );
+	Options.ClassFilters.Add(MakeShareable( new FBehaviorTreeEditorUtils::FNewNodeClassFilter<UBTTask_BlueprintBase> ));
 
 	FOnClassPicked OnPicked( FOnClassPicked::CreateSP( this, &FBehaviorTreeEditor::HandleNewNodeClassPicked ) );
 
@@ -1705,7 +1705,7 @@ TSharedRef<SWidget> FBehaviorTreeEditor::HandleCreateNewDecoratorMenu() const
 {
 	FClassViewerInitializationOptions Options;
 	Options.bShowUnloadedBlueprints = true;
-	Options.ClassFilter = MakeShareable( new FBehaviorTreeEditorUtils::FNewNodeClassFilter<UBTDecorator_BlueprintBase> );
+	Options.ClassFilters.Add(MakeShareable( new FBehaviorTreeEditorUtils::FNewNodeClassFilter<UBTDecorator_BlueprintBase> ));
 
 	FOnClassPicked OnPicked( FOnClassPicked::CreateSP( this, &FBehaviorTreeEditor::HandleNewNodeClassPicked ) );
 
@@ -1716,7 +1716,7 @@ TSharedRef<SWidget> FBehaviorTreeEditor::HandleCreateNewServiceMenu() const
 {
 	FClassViewerInitializationOptions Options;
 	Options.bShowUnloadedBlueprints = true;
-	Options.ClassFilter = MakeShareable( new FBehaviorTreeEditorUtils::FNewNodeClassFilter<UBTService_BlueprintBase> );
+	Options.ClassFilters.Add(MakeShareable( new FBehaviorTreeEditorUtils::FNewNodeClassFilter<UBTService_BlueprintBase> ));
 
 	FOnClassPicked OnPicked( FOnClassPicked::CreateSP( this, &FBehaviorTreeEditor::HandleNewNodeClassPicked ) );
 
diff --git a/Engine/Source/Editor/BehaviorTreeEditor/Private/SBehaviorTreeBlackboardEditor.cpp b/Engine/Source/Editor/BehaviorTreeEditor/Private/SBehaviorTreeBlackboardEditor.cpp
--- a/Engine/Source/Editor/BehaviorTreeEditor/Private/SBehaviorTreeBlackboardEditor.cpp
+++ b/Engine/Source/Editor/BehaviorTreeEditor/Private/SBehaviorTreeBlackboardEditor.cpp
@@ -211,7 +211,7 @@ TSharedRef<SWidget> SBehaviorTreeBlackboardEditor::HandleCreateNewEntryMenu() co
 	FClassViewerInitializationOptions Options;
 	Options.bShowUnloadedBlueprints = true;
 	Options.NameTypeToDisplay = EClassViewerNameTypeToDisplay::DisplayName;
-	Options.ClassFilter = MakeShareable( new FBlackboardEntryClassFilter );
+	Options.ClassFilters.Add(MakeShareable( new FBlackboardEntryClassFilter ));
 
 	FOnClassPicked OnPicked( FOnClassPicked::CreateRaw( const_cast<SBehaviorTreeBlackboardEditor*>(this), &SBehaviorTreeBlackboardEditor::HandleKeyClassPicked ) );
 
diff --git a/Engine/Source/Editor/Blutility/Private/EditorUtilityBlueprintFactory.cpp b/Engine/Source/Editor/Blutility/Private/EditorUtilityBlueprintFactory.cpp
--- a/Engine/Source/Editor/Blutility/Private/EditorUtilityBlueprintFactory.cpp
+++ b/Engine/Source/Editor/Blutility/Private/EditorUtilityBlueprintFactory.cpp
@@ -81,7 +81,7 @@ bool UEditorUtilityBlueprintFactory::ConfigureProperties()
 
 	TSharedPtr< FBlutilityBlueprintFactoryFilter > Filter = MakeShareable(new FBlutilityBlueprintFactoryFilter);
 	Filter->DisallowedChildOfClasses.Add(UEditorUtilityWidget::StaticClass());
-	Options.ClassFilter = Filter;
+	Options.ClassFilters.Add(Filter.ToSharedRef());
 
 	const FText TitleText = NSLOCTEXT("EditorFactories", "CreateBlueprintOptions", "Pick Parent Class");
 	UClass* ChosenClass = NULL;
diff --git a/Engine/Source/Editor/Blutility/Private/EditorUtilityWidgetBlueprintFactory.cpp b/Engine/Source/Editor/Blutility/Private/EditorUtilityWidgetBlueprintFactory.cpp
--- a/Engine/Source/Editor/Blutility/Private/EditorUtilityWidgetBlueprintFactory.cpp
+++ b/Engine/Source/Editor/Blutility/Private/EditorUtilityWidgetBlueprintFactory.cpp
@@ -77,7 +77,7 @@ bool UEditorUtilityWidgetBlueprintFactory::ConfigureProperties()
 		Options.ExtraPickerCommonClasses.Add(UCanvasPanel::StaticClass());
 
 		TSharedPtr<FEditorUtilityWidgetBlueprintFactoryFilter> Filter = MakeShareable(new FEditorUtilityWidgetBlueprintFactoryFilter);
-		Options.ClassFilter = Filter;
+		Options.ClassFilters.Add(Filter.ToSharedRef());
 
 		Filter->DisallowedClassFlags = CLASS_Abstract | CLASS_Deprecated | CLASS_NewerVersionExists;
 		Filter->AllowedChildrenOfClasses.Add(UPanelWidget::StaticClass());
diff --git a/Engine/Source/Editor/ClassViewer/Private/ClassViewerFilter.cpp b/Engine/Source/Editor/ClassViewer/Private/ClassViewerFilter.cpp
--- a/Engine/Source/Editor/ClassViewer/Private/ClassViewerFilter.cpp
+++ b/Engine/Source/Editor/ClassViewer/Private/ClassViewerFilter.cpp
@@ -488,10 +488,22 @@ bool FClassViewerFilter::IsClassAllowed(const FClassViewerInitializationOptions&
 	}
 
 	bool bPassesCustomFilter = true;
-	if (InInitOptions.ClassFilter.IsValid())
+	for (const TSharedRef<IClassViewerFilter>& CustomFilter : InInitOptions.ClassFilters)
+	{
+		if (!CustomFilter->IsClassAllowed(InInitOptions, InClass, FilterFunctions))
+		{
+			bPassesCustomFilter = false;
+			break;
+		}
+	}
+
+PRAGMA_DISABLE_DEPRECATION_WARNINGS
+	// Retained for backcompat; can remove when fully deprecated.
+	if (bPassesCustomFilter && InInitOptions.ClassFilter.IsValid())
 	{
 		bPassesCustomFilter = InInitOptions.ClassFilter->IsClassAllowed(InInitOptions, InClass, FilterFunctions);
 	}
+PRAGMA_ENABLE_DEPRECATION_WARNINGS
 
 	bool bPassesTextFilter = true;
 	if (bCheckTextFilter && (TextFilter->GetFilterType() != ETextFilterExpressionType::Empty))
@@ -590,10 +602,22 @@ bool FClassViewerFilter::IsUnloadedClassAllowed(const FClassViewerInitialization
 	}
 
 	bool bPassesCustomFilter = true;
-	if (InInitOptions.ClassFilter.IsValid())
+	for (const TSharedRef<IClassViewerFilter>& CustomFilter : InInitOptions.ClassFilters)
+	{
+		if (!CustomFilter->IsUnloadedClassAllowed(InInitOptions, InUnloadedClassData, FilterFunctions))
+		{
+			bPassesCustomFilter = false;
+			break;
+		}
+	}
+
+PRAGMA_DISABLE_DEPRECATION_WARNINGS
+	// Retained for backcompat; can remove when fully deprecated.
+	if (bPassesCustomFilter && InInitOptions.ClassFilter.IsValid())
 	{
 		bPassesCustomFilter = InInitOptions.ClassFilter->IsUnloadedClassAllowed(InInitOptions, InUnloadedClassData, FilterFunctions);
 	}
+PRAGMA_ENABLE_DEPRECATION_WARNINGS
 
 	const bool bPassesTextFilter = PassesTextFilter(*InUnloadedClassData->GetClassName().Get(), TextFilter);
 
@@ -609,4 +633,4 @@ bool FClassViewerFilter::IsUnloadedClassAllowed(const FClassViewerInitialization
 		&& (!bCheckTextFilter || bPassesTextFilter) && bPassesAssetReferenceFilter;
 
 	return bPassesFilter;
-}
+}
\ No newline at end of file
diff --git a/Engine/Source/Editor/ClassViewer/Private/SClassViewer.cpp b/Engine/Source/Editor/ClassViewer/Private/SClassViewer.cpp
--- a/Engine/Source/Editor/ClassViewer/Private/SClassViewer.cpp
+++ b/Engine/Source/Editor/ClassViewer/Private/SClassViewer.cpp
@@ -1473,6 +1473,21 @@ void SClassViewer::Construct(const FArguments& InArgs, const FClassViewerInitial
 		}
 	}
 
+	// Clear out the current set of custom filter options.
+	CustomClassFilterOptions.Empty(InitOptions.ClassFilters.Num());
+
+	// Gather additional filter options from any custom filters.
+	TArray<TSharedRef<FClassViewerFilterOption>> FilterOptions;
+	for (TSharedRef<IClassViewerFilter> CustomFilter : InitOptions.ClassFilters)
+	{
+		// Append this filter's options to the current set.
+		CustomFilter->GetFilterOptions(FilterOptions);
+		CustomClassFilterOptions.Append(FilterOptions);
+
+		// Clear out the temp array for the next pass.
+		FilterOptions.Empty();
+	}
+
 	TSharedRef<SWidget> FiltersWidget = SNullWidget::NullWidget;
 	// Build the top menu
 	if(InitOptions.Mode == EClassViewerMode::ClassBrowsing)
@@ -1932,6 +1947,22 @@ TSharedRef<SWidget> SClassViewer::GetViewButtonContent()
 			NAME_None,
 			EUserInterfaceActionType::ToggleButton
 			);
+
+		for (const TSharedRef<FClassViewerFilterOption>& FilterOption : CustomClassFilterOptions)
+		{
+			MenuBuilder.AddMenuEntry(
+				FilterOption->LabelText,
+				FilterOption->ToolTipText,
+				FSlateIcon(),
+				FUIAction(
+					FExecuteAction::CreateSP(this, &SClassViewer::ToggleCustomFilterOption, FilterOption),
+					FCanExecuteAction(),
+					FIsActionChecked::CreateSP(this, &SClassViewer::IsCustomFilterOptionEnabled, FilterOption)
+				),
+				NAME_None,
+				EUserInterfaceActionType::ToggleButton
+			);
+		}
 	}
 	MenuBuilder.EndSection();
 
@@ -2665,4 +2696,21 @@ bool SClassViewer::IsShowingInternalClasses() const
 	return IsToggleShowInternalClassesAllowed() ? GetDefault<UClassViewerSettings>()->DisplayInternalClasses : false;
 }
 
+void SClassViewer::ToggleCustomFilterOption(TSharedRef<FClassViewerFilterOption> FilterOption)
+{
+	FilterOption->bEnabled = !FilterOption->bEnabled;
+
+	if (FilterOption->OnOptionChanged.IsBound())
+	{
+		FilterOption->OnOptionChanged.Execute(FilterOption->bEnabled);
+	}
+
+	Refresh();
+}
+
+bool SClassViewer::IsCustomFilterOptionEnabled(TSharedRef<FClassViewerFilterOption> FilterOption) const
+{
+	return FilterOption->bEnabled;
+}
+
 #undef LOCTEXT_NAMESPACE
diff --git a/Engine/Source/Editor/ClassViewer/Public/ClassViewerFilter.h b/Engine/Source/Editor/ClassViewer/Public/ClassViewerFilter.h
--- a/Engine/Source/Editor/ClassViewer/Public/ClassViewerFilter.h
+++ b/Engine/Source/Editor/ClassViewer/Public/ClassViewerFilter.h
@@ -11,6 +11,26 @@ class FTextFilterExpressionEvaluator;
 class IAssetReferenceFilter;
 class IAssetRegistry;
 
+/** Delegate used to respond to a filter option change. The argument will indicate whether the option was enabled or disabled. */
+DECLARE_DELEGATE_OneParam(FOnClassViewerFilterOptionChanged, bool);
+
+/** Used to define a custom class viewer filter option. */
+class FClassViewerFilterOption
+{
+public:
+	/** Whether or not the option is currently enabled (default = true). */
+	bool bEnabled = true;
+
+	/** Localized label text to use for the option in the class viewer's filter menu. */
+	TAttribute<FText> LabelText;
+
+	/** Localized tooltip text to use for the option in the class viewer's filter menu. */
+	TAttribute<FText> ToolTipText;
+
+	/** Optional external delegate that will be invoked when this filter option is changed. */
+	FOnClassViewerFilterOptionChanged OnOptionChanged;
+};
+
 /** Interface class for creating filters for the Class Viewer. */
 class IClassViewerFilter
 {
@@ -34,6 +54,48 @@ public:
 	 * @param InFilterFuncs				Useful functions for filtering.
 	 */
 	virtual bool IsUnloadedClassAllowed(const FClassViewerInitializationOptions& InInitOptions, const TSharedRef< const class IUnloadedBlueprintData > InUnloadedClassData, TSharedRef< class FClassViewerFilterFuncs > InFilterFuncs) = 0;
+
+	/**
+	 * Can be optionally implemented to gather additional filter flags. If present, these
+	 * will be added into the View Options menu where they can be toggled on/off by the user.
+	 * 
+	 * Example:
+	 *	In the following case, we define a custom filter option to allow the user to toggle the entire filter on/off. This is
+	 *	synced with a global config setting (MyFilterConfigSettings->bIsCustomFilterEnabled) that is checked at filtering time.
+	 * 
+	 *		void FMyCustomClassFilter::GetFilterOptions(TArray<TSharedRef<FClassViewerFilterOption>>& OutFilterOptions)
+	 *		{
+	 *			TSharedRef<FClassViewerFilterOption> MyFilterOption = MakeShared<FClassViewerFilterOption>();
+	 *			MyFilterOption->bEnabled = MyFilterConfigSettings->bIsCustomFilterEnabled;
+	 *			MyFilterOption->LabelText = LOCTEXT("MyCustomFilterOptionLabel", "My Custom Filter");
+	 *			MyFilterOption->ToolTipText = LOCTEXT("MyCustomFilterOptionToolTip", "Enable or disable my custom class filter.");
+	 *			MyFilterOption->OnOptionChanged = FOnClassViewerFilterOptionChanged::CreateSP(this, &FMyCustomClassFilter::OnOptionChanged);
+	 * 
+	 *			OutFilterOptions.Add(MyFilterOption);
+	 *		}
+	 *		
+	 *		void FMyCustomClassFilter::OnOptionChanged(bool bIsEnabled)
+	 *		{
+	 *			// Updates the filter's config setting whenever the user changes the option.
+	 *			MyFilterConfigSettings->bIsCustomFilterEnabled = bIsEnabled;
+	 *		}
+	 * 
+	 *		bool FMyCustomClassFilter::IsClassAllowed(...)
+	 *		{
+	 *			if(!MyFilterConfigSettings->bIsCustomFilterEnabled)
+	 *			{
+	 *				// Filter is disabled; always allow the class.
+	 *				return true;
+	 *			}
+	 * 
+	 *			// ...
+	 *		}
+	 *
+	 * @param OutFilterOptions		On output, contains the set of options to be added into the filter menu.
+	 *								The Class Viewer requires these to be allocated so that each option can
+	 *								be safely referenced by each menu item widget in the View Options menu.
+	 */
+	virtual void GetFilterOptions(TArray<TSharedRef<FClassViewerFilterOption>>& OutFilterOptions) {}
 };
 
 /** Filter class that performs many common checks. */
diff --git a/Engine/Source/Editor/ClassViewer/Public/ClassViewerModule.h b/Engine/Source/Editor/ClassViewer/Public/ClassViewerModule.h
--- a/Engine/Source/Editor/ClassViewer/Public/ClassViewerModule.h
+++ b/Engine/Source/Editor/ClassViewer/Public/ClassViewerModule.h
@@ -57,13 +57,18 @@ enum class EClassViewerNameTypeToDisplay : uint8
  * Settings for the Class Viewer set by the programmer before spawning an instance of the widget.  This
  * is used to modify the class viewer's behavior in various ways, such as filtering in or out specific classes.
  */
+PRAGMA_DISABLE_DEPRECATION_WARNINGS
 class FClassViewerInitializationOptions
 {
 
 public:
-	/** The filter to use on classes in this instance. */
+	/** [Deprecated] The filter to use on classes in this instance. */
+	UE_DEPRECATED(5.0, "Please add to the ClassFilters array member instead.")
 	TSharedPtr<class IClassViewerFilter> ClassFilter;
 
+	/** The filter(s) to use on classes in this instance. */
+	TArray<TSharedRef<class IClassViewerFilter>> ClassFilters;
+
 	/** Mode to operate in */
 	EClassViewerMode::Type Mode;
 
@@ -146,6 +151,7 @@ public:
 	{
 	}
 };
+PRAGMA_ENABLE_DEPRECATION_WARNINGS
 
 /**
  * Class Viewer module
diff --git a/Engine/Source/Editor/ClassViewer/Public/SClassViewer.h b/Engine/Source/Editor/ClassViewer/Public/SClassViewer.h
--- a/Engine/Source/Editor/ClassViewer/Public/SClassViewer.h
+++ b/Engine/Source/Editor/ClassViewer/Public/SClassViewer.h
@@ -22,6 +22,7 @@ class FTextFilterExpressionEvaluator;
 class UBlueprint;
 class SComboButton;
 class FClassViewerFilter;
+class FClassViewerFilterOption;
 
 //////////////////////////////////////////////////////////////////////////
 // SClassViewer
@@ -206,6 +207,12 @@ private:
 	/** Whether or not it's possible to show internal use classes */
 	bool IsToggleShowInternalClassesAllowed() const;
 
+	/** Toggle whether a custom class filter option is enabled */
+	void ToggleCustomFilterOption(TSharedRef<FClassViewerFilterOption> FilterOption);
+
+	/** Returns whether a custom class filter option is currently enabled */
+	bool IsCustomFilterOptionEnabled(TSharedRef<FClassViewerFilterOption> FilterOption) const;
+
 	/** Get the total number of classes passing the current filters.*/
 	const int GetNumItems() const;
 
@@ -274,4 +281,7 @@ private:
 
 	/** Number of classes that passed the filter*/
 	int NumClasses;
+
+	/** Custom filter options for the View Options menu */
+	TArray<TSharedRef<FClassViewerFilterOption>> CustomClassFilterOptions;
 };
diff --git a/Engine/Source/Editor/CurveEditor/Private/Filters/SCurveEditorFilterPanel.cpp b/Engine/Source/Editor/CurveEditor/Private/Filters/SCurveEditorFilterPanel.cpp
--- a/Engine/Source/Editor/CurveEditor/Private/Filters/SCurveEditorFilterPanel.cpp
+++ b/Engine/Source/Editor/CurveEditor/Private/Filters/SCurveEditorFilterPanel.cpp
@@ -46,7 +46,7 @@ void SCurveEditorFilterPanel::Construct(const FArguments& InArgs, TSharedRef<FCu
 	Options.InitiallySelectedClass = DefaultFilterClass;
 
 	TSharedPtr<FCurveFilterClassFilter> ClassFilter = MakeShared<FCurveFilterClassFilter>();
-	Options.ClassFilter = ClassFilter;
+	Options.ClassFilters.Add(ClassFilter.ToSharedRef());
 
 	FOnClassPicked OnPicked(FOnClassPicked::CreateSP(this, &SCurveEditorFilterPanel::SetFilterClass));
 
diff --git a/Engine/Source/Editor/DetailCustomizations/Private/ActorDetails.cpp b/Engine/Source/Editor/DetailCustomizations/Private/ActorDetails.cpp
--- a/Engine/Source/Editor/DetailCustomizations/Private/ActorDetails.cpp
+++ b/Engine/Source/Editor/DetailCustomizations/Private/ActorDetails.cpp
@@ -285,9 +285,9 @@ void FActorDetails::CreateClassPickerConvertActorFilter(const TWeakObjectPtr<AAc
 {
 	// Shouldn't ever be overwriting an already established filter
 	check( ConvertActor.IsValid() )
-	check( ClassPickerOptions != NULL && !ClassPickerOptions->ClassFilter.IsValid() )
-	TSharedPtr<FConvertToClassFilter> Filter = MakeShareable(new FConvertToClassFilter);
-	ClassPickerOptions->ClassFilter = Filter;
+	check( ClassPickerOptions != nullptr && ClassPickerOptions->ClassFilters.IsEmpty() );
+	TSharedRef<FConvertToClassFilter> Filter = MakeShared<FConvertToClassFilter>();
+	ClassPickerOptions->ClassFilters.Add(Filter);
 
 	UClass* ConvertClass = ConvertActor->GetClass();
 	UClass* RootConversionClass = GetConversionRoot(ConvertClass);
diff --git a/Engine/Source/Editor/DetailCustomizations/Private/BrushDetails.cpp b/Engine/Source/Editor/DetailCustomizations/Private/BrushDetails.cpp
--- a/Engine/Source/Editor/DetailCustomizations/Private/BrushDetails.cpp
+++ b/Engine/Source/Editor/DetailCustomizations/Private/BrushDetails.cpp
@@ -257,7 +257,7 @@ TSharedRef<SWidget> FBrushDetails::GenerateBuildMenuContent()
 	};
 
 	FClassViewerInitializationOptions Options;
-	Options.ClassFilter = MakeShareable(new FBrushFilter);
+	Options.ClassFilters.Add(MakeShareable(new FBrushFilter));
 	Options.Mode = EClassViewerMode::ClassPicker;
 	Options.DisplayMode = EClassViewerDisplayMode::ListView;
 	return FModuleManager::LoadModuleChecked<FClassViewerModule>("ClassViewer").CreateClassViewer(Options, FOnClassPicked::CreateSP(this, &FBrushDetails::OnClassPicked));
diff --git a/Engine/Source/Editor/DetailCustomizations/Private/SkeletalMeshComponentDetails.cpp b/Engine/Source/Editor/DetailCustomizations/Private/SkeletalMeshComponentDetails.cpp
--- a/Engine/Source/Editor/DetailCustomizations/Private/SkeletalMeshComponentDetails.cpp
+++ b/Engine/Source/Editor/DetailCustomizations/Private/SkeletalMeshComponentDetails.cpp
@@ -299,7 +299,7 @@ TSharedRef<SWidget> FSkeletalMeshComponentDetails::GetClassPickerMenuContent()
 	FClassViewerModule& ClassViewerModule = FModuleManager::LoadModuleChecked<FClassViewerModule>("ClassViewer");
 	FClassViewerInitializationOptions InitOptions;
 	InitOptions.Mode = EClassViewerMode::ClassPicker;
-	InitOptions.ClassFilter = Filter;
+	InitOptions.ClassFilters.Add(Filter.ToSharedRef());
 	InitOptions.bShowNoneOption = true;
 
 	return SNew(SBorder)
diff --git a/Engine/Source/Editor/GameProjectGeneration/Private/SNewClassDialog.cpp b/Engine/Source/Editor/GameProjectGeneration/Private/SNewClassDialog.cpp
--- a/Engine/Source/Editor/GameProjectGeneration/Private/SNewClassDialog.cpp
+++ b/Engine/Source/Editor/GameProjectGeneration/Private/SNewClassDialog.cpp
@@ -214,20 +214,22 @@ void SNewClassDialog::Construct( const FArguments& InArgs )
 	Options.bShowObjectRootClass = true;
 	Options.bExpandRootNodes = true;
 
-	if (InArgs._ClassViewerFilter.IsValid())
-	{
-		Options.ClassFilter = InArgs._ClassViewerFilter;
-	}
-	else if (InArgs._ClassDomain == EClassDomain::Native)
+	TSharedPtr<IClassViewerFilter> ClassFilter = InArgs._ClassViewerFilter;
+	if (!ClassFilter.IsValid() && InArgs._ClassDomain == EClassDomain::Native)
 	{
 		// Prevent creating native classes based on blueprint classes
-		Options.ClassFilter = MakeShareable(new FNativeClassParentFilter());
+		ClassFilter = MakeShared<FNativeClassParentFilter>();
 	}
 
-	// Only show the Object root class if it's a valid base (this helps keep the tree clean)
-	if (Options.ClassFilter.IsValid() && !Options.ClassFilter->IsClassAllowed(Options, UObject::StaticClass(), MakeShareable(new FClassViewerFilterFuncs)))
+	if (ClassFilter.IsValid())
 	{
-		Options.bShowObjectRootClass = false;
+		Options.ClassFilters.Add(ClassFilter.ToSharedRef());
+
+		// Only show the Object root class if it's a valid base (this helps keep the tree clean)
+		if (!ClassFilter->IsClassAllowed(Options, UObject::StaticClass(), MakeShared<FClassViewerFilterFuncs>()))
+		{
+			Options.bShowObjectRootClass = false;
+		}
 	}
 
 	ClassViewer = StaticCastSharedRef<SClassViewer>(FModuleManager::LoadModuleChecked<FClassViewerModule>("ClassViewer").CreateClassViewer(Options, FOnClassPicked::CreateSP(this, &SNewClassDialog::OnAdvancedClassSelected)));
diff --git a/Engine/Source/Editor/GraphEditor/Private/KismetNodes/SGraphNodeSpawnActor.cpp b/Engine/Source/Editor/GraphEditor/Private/KismetNodes/SGraphNodeSpawnActor.cpp
--- a/Engine/Source/Editor/GraphEditor/Private/KismetNodes/SGraphNodeSpawnActor.cpp
+++ b/Engine/Source/Editor/GraphEditor/Private/KismetNodes/SGraphNodeSpawnActor.cpp
@@ -79,7 +79,7 @@ protected:
 		Options.bShowNoneOption = true;
 		Options.bShowObjectRootClass = false;
 		TSharedPtr< FActorBasedBlueprintClassFilter > Filter = MakeShareable(new FActorBasedBlueprintClassFilter);
-		Options.ClassFilter = Filter;
+		Options.ClassFilters.Add(Filter.ToSharedRef());
 
 		return
 			SNew(SBox)
diff --git a/Engine/Source/Editor/GraphEditor/Private/KismetNodes/SGraphNodeSpawnActorFromClass.cpp b/Engine/Source/Editor/GraphEditor/Private/KismetNodes/SGraphNodeSpawnActorFromClass.cpp
--- a/Engine/Source/Editor/GraphEditor/Private/KismetNodes/SGraphNodeSpawnActorFromClass.cpp
+++ b/Engine/Source/Editor/GraphEditor/Private/KismetNodes/SGraphNodeSpawnActorFromClass.cpp
@@ -65,7 +65,7 @@ protected:
 		Options.bShowNoneOption = true;
 		Options.bShowObjectRootClass = true;
 		TSharedPtr< FActorBasedClassFilter > Filter = MakeShareable(new FActorBasedClassFilter);
-		Options.ClassFilter = Filter;
+		Options.ClassFilters.Add(Filter.ToSharedRef());
 		// Populate the referencing asset, if possible
 		if (UEdGraphPin* LocalGraphPin = GetPinObj())
 		{
diff --git a/Engine/Source/Editor/GraphEditor/Private/KismetPins/SGraphPinClass.cpp b/Engine/Source/Editor/GraphEditor/Private/KismetPins/SGraphPinClass.cpp
--- a/Engine/Source/Editor/GraphEditor/Private/KismetPins/SGraphPinClass.cpp
+++ b/Engine/Source/Editor/GraphEditor/Private/KismetPins/SGraphPinClass.cpp
@@ -120,7 +120,7 @@ TSharedRef<SWidget> SGraphPinClass::GenerateAssetPicker()
 		Filter->bAllowAbstractClasses = AllowAbstractString.ToBool();
 	}
 
-	Options.ClassFilter = Filter;
+	Options.ClassFilters.Add(Filter.ToSharedRef());
 
 	Filter->AllowedChildrenOfClasses.Add(PinRequiredParentClass);
 	Filter->GraphPinOutermostPackage = GraphPinObj->GetOuter()->GetOutermost();
diff --git a/Engine/Source/Editor/Kismet/Private/SBlueprintLibraryPalette.cpp b/Engine/Source/Editor/Kismet/Private/SBlueprintLibraryPalette.cpp
--- a/Engine/Source/Editor/Kismet/Private/SBlueprintLibraryPalette.cpp
+++ b/Engine/Source/Editor/Kismet/Private/SBlueprintLibraryPalette.cpp
@@ -366,7 +366,7 @@ TSharedRef<SWidget> SBlueprintLibraryPalette::ConstructClassFilterDropdownConten
 	FClassViewerInitializationOptions Options;
 	Options.Mode        = EClassViewerMode::ClassPicker;
 	Options.DisplayMode = EClassViewerDisplayMode::TreeView;
-	Options.ClassFilter = MakeShareable(new FPaletteClassFilter);
+	Options.ClassFilters.Add(MakeShareable(new FPaletteClassFilter));
 	//  create a class picker for the drop-down
 	TSharedRef<SWidget> ClassPickerWidget = FModuleManager::LoadModuleChecked<FClassViewerModule>("ClassViewer").CreateClassViewer(Options, FOnClassPicked::CreateSP(this, &SBlueprintLibraryPalette::OnClassPicked));
 
diff --git a/Engine/Source/Editor/KismetWidgets/Private/CreateBlueprintFromActorDialog.cpp b/Engine/Source/Editor/KismetWidgets/Private/CreateBlueprintFromActorDialog.cpp
--- a/Engine/Source/Editor/KismetWidgets/Private/CreateBlueprintFromActorDialog.cpp
+++ b/Engine/Source/Editor/KismetWidgets/Private/CreateBlueprintFromActorDialog.cpp
@@ -289,7 +289,7 @@ void SSCreateBlueprintPicker::Construct(const FArguments& InArgs)
 	{
 		ActorOverrideClass = ActorOverride->GetClass();
 		TSharedPtr<FBlueprintFromActorParentFilter> Filter = MakeShareable(new FBlueprintFromActorParentFilter(ActorOverrideClass, CreateMode));
-		ClassViewerOptions.ClassFilter = Filter;
+		ClassViewerOptions.ClassFilters.Add(Filter.ToSharedRef());
 	}
 
 	if (ActorOverrideClass && CreateMode == ECreateBlueprintFromActorMode::Subclass)
diff --git a/Engine/Source/Editor/LevelEditor/Private/LevelEditorToolBar.cpp b/Engine/Source/Editor/LevelEditor/Private/LevelEditorToolBar.cpp
--- a/Engine/Source/Editor/LevelEditor/Private/LevelEditorToolBar.cpp
+++ b/Engine/Source/Editor/LevelEditor/Private/LevelEditorToolBar.cpp
@@ -324,7 +324,7 @@ namespace LevelEditorActionHelpers
 
 		TSharedPtr< FBlueprintParentFilter_MapModeSettings > Filter = MakeShareable(new FBlueprintParentFilter_MapModeSettings);
 		Filter->AllowedChildrenOfClasses.Add(InRootClass);
-		Options.ClassFilter = Filter;
+		Options.ClassFilters.Add(Filter.ToSharedRef());
 
 		FText RootClassName = FText::FromString(InRootClass->GetName());
 		TSharedRef<SWidget> ClassViewer = FModuleManager::LoadModuleChecked<FClassViewerModule>("ClassViewer").CreateClassViewer(Options, InOnClassPicked);
@@ -356,7 +356,7 @@ namespace LevelEditorActionHelpers
 
 		TSharedPtr< FBlueprintParentFilter_MapModeSettings > Filter = MakeShareable(new FBlueprintParentFilter_MapModeSettings);
 		Filter->AllowedChildrenOfClasses.Add(InRootClass);
-		Options.ClassFilter = Filter;
+		Options.ClassFilters.Add(Filter.ToSharedRef());
 
 		FText RootClassName = FText::FromString(InRootClass->GetName());
 		TSharedRef<SWidget> ClassViewer = FModuleManager::LoadModuleChecked<FClassViewerModule>("ClassViewer").CreateClassViewer(Options, InOnClassPicked);
diff --git a/Engine/Source/Editor/LevelInstanceEditor/Private/LevelInstanceEditorModule.cpp b/Engine/Source/Editor/LevelInstanceEditor/Private/LevelInstanceEditorModule.cpp
--- a/Engine/Source/Editor/LevelInstanceEditor/Private/LevelInstanceEditorModule.cpp
+++ b/Engine/Source/Editor/LevelInstanceEditor/Private/LevelInstanceEditorModule.cpp
@@ -653,7 +653,7 @@ namespace LevelInstanceMenuUtils
 			Options->bIsBlueprintBaseOnly = false;
 			Options->InitiallySelectedClass = ALevelInstance::StaticClass();
 			Options->bIsActorsOnly = true;
-			Options->ClassFilter = MakeShareable(new FLevelInstanceClassFilter);
+			Options->ClassFilters.Add(MakeShareable(new FLevelInstanceClassFilter));
 		});
 		ON_SCOPE_EXIT
 		{
diff --git a/Engine/Source/Editor/Persona/Private/SAnimNotifyPanel.cpp b/Engine/Source/Editor/Persona/Private/SAnimNotifyPanel.cpp
--- a/Engine/Source/Editor/Persona/Private/SAnimNotifyPanel.cpp
+++ b/Engine/Source/Editor/Persona/Private/SAnimNotifyPanel.cpp
@@ -2349,7 +2349,7 @@ void SAnimNotifyTrack::MakeNewNotifyPicker(FMenuBuilder& MenuBuilder, bool bIsRe
 	InitOptions.bEnableClassDynamicLoading = true;
 	InitOptions.bExpandRootNodes = true;
 	InitOptions.NameTypeToDisplay = EClassViewerNameTypeToDisplay::DisplayName;
-	InitOptions.ClassFilter = MakeShared<FNotifyStateClassFilter>(Sequence);
+	InitOptions.ClassFilters.Add(MakeShared<FNotifyStateClassFilter>(Sequence));
 	InitOptions.bShowBackgroundBorder = false;
 
 	FClassViewerModule& ClassViewerModule = FModuleManager::LoadModuleChecked<FClassViewerModule>("ClassViewer");
diff --git a/Engine/Source/Editor/PropertyEditor/Private/IDetailsViewPrivate.h b/Engine/Source/Editor/PropertyEditor/Private/IDetailsViewPrivate.h
--- a/Engine/Source/Editor/PropertyEditor/Private/IDetailsViewPrivate.h
+++ b/Engine/Source/Editor/PropertyEditor/Private/IDetailsViewPrivate.h
@@ -95,6 +95,11 @@ public:
 	 */
 	virtual TSharedPtr<class FAssetThumbnailPool> GetThumbnailPool() const = 0;
 
+	/**
+	 * @return The set of custom class viewer filters to use for class properties in this view
+	 */
+	virtual const TArray<TSharedRef<class IClassViewerFilter>>& GetClassViewerFilters() const = 0;
+
 	/**
 	 * Creates the color picker window for this property view.
 	 *
diff --git a/Engine/Source/Editor/PropertyEditor/Private/Presentation/PropertyTable/PropertyTable.cpp b/Engine/Source/Editor/PropertyEditor/Private/Presentation/PropertyTable/PropertyTable.cpp
--- a/Engine/Source/Editor/PropertyEditor/Private/Presentation/PropertyTable/PropertyTable.cpp
+++ b/Engine/Source/Editor/PropertyEditor/Private/Presentation/PropertyTable/PropertyTable.cpp
@@ -113,6 +113,13 @@ TSharedPtr<class FAssetThumbnailPool> FPropertyTable::GetThumbnailPool() const
 	return NULL;
 }
 
+const TArray<TSharedRef<class IClassViewerFilter>>& FPropertyTable::GetClassViewerFilters() const
+{
+	// not implemented
+	static TArray<TSharedRef<class IClassViewerFilter>> NotImplemented;
+	return NotImplemented;
+}
+
 bool FPropertyTable::GetIsUserAllowedToChangeRoot()
 {
 	return AllowUserToChangeRoot;
diff --git a/Engine/Source/Editor/PropertyEditor/Private/Presentation/PropertyTable/PropertyTable.h b/Engine/Source/Editor/PropertyEditor/Private/Presentation/PropertyTable/PropertyTable.h
--- a/Engine/Source/Editor/PropertyEditor/Private/Presentation/PropertyTable/PropertyTable.h
+++ b/Engine/Source/Editor/PropertyEditor/Private/Presentation/PropertyTable/PropertyTable.h
@@ -29,6 +29,7 @@ public:
 	virtual void EnqueueDeferredAction( FSimpleDelegate DeferredAction ) override;
 	virtual TSharedPtr<class FAssetThumbnailPool> GetThumbnailPool() const override;
 	virtual void NotifyFinishedChangingProperties(const FPropertyChangedEvent& PropertyChangedEvent) override {}
+	virtual const TArray<TSharedRef<class IClassViewerFilter>>& GetClassViewerFilters() const override;
 
 	virtual bool GetIsUserAllowedToChangeRoot() override;
 	virtual void SetIsUserAllowedToChangeRoot( bool InAllowUserToChangeRoot ) override;
diff --git a/Engine/Source/Editor/PropertyEditor/Private/PropertyEditorHelpers.cpp b/Engine/Source/Editor/PropertyEditor/Private/PropertyEditorHelpers.cpp
--- a/Engine/Source/Editor/PropertyEditor/Private/PropertyEditorHelpers.cpp
+++ b/Engine/Source/Editor/PropertyEditor/Private/PropertyEditorHelpers.cpp
@@ -155,9 +155,11 @@ TSharedRef<SWidget> SPropertyValueWidget::ConstructPropertyEditorWidget( TShared
 		}
 		else if (SPropertyEditorClass::Supports(PropertyEditorRef))
 		{
+			static TArray<TSharedRef<class IClassViewerFilter>> NullFilters;
 			TSharedRef<SPropertyEditorClass> ClassWidget =
 				SAssignNew(PropertyWidget, SPropertyEditorClass, PropertyEditorRef)
-				.Font(FontStyle);
+				.Font(FontStyle)
+				.ClassViewerFilters(InPropertyUtilities.IsValid() ? InPropertyUtilities->GetClassViewerFilters() : NullFilters);
 
 			ClassWidget->GetDesiredWidth(MinDesiredWidth, MaxDesiredWidth);
 		}
diff --git a/Engine/Source/Editor/PropertyEditor/Private/PropertyRowGenerator.cpp b/Engine/Source/Editor/PropertyEditor/Private/PropertyRowGenerator.cpp
--- a/Engine/Source/Editor/PropertyEditor/Private/PropertyRowGenerator.cpp
+++ b/Engine/Source/Editor/PropertyEditor/Private/PropertyRowGenerator.cpp
@@ -93,6 +93,19 @@ public:
 		return Generator != nullptr && Generator->HasClassDefaultObject();
 	}
 
+	virtual const TArray<TSharedRef<class IClassViewerFilter>>& GetClassViewerFilters() const override
+	{
+		if (Generator != nullptr)
+		{
+			return Generator->GetClassViewerFilters();
+		}
+		else
+		{
+			static TArray<TSharedRef<class IClassViewerFilter>> NullFilters;
+			return NullFilters;
+		}
+	}
+
 private:
 	FPropertyRowGenerator* Generator;
 };
@@ -422,6 +435,13 @@ TSharedPtr<class FAssetThumbnailPool> FPropertyRowGenerator::GetThumbnailPool()
 	return UThumbnailManager::Get().GetSharedThumbnailPool();
 }
 
+const TArray<TSharedRef<class IClassViewerFilter>>& FPropertyRowGenerator::GetClassViewerFilters() const
+{
+	// not implemented
+	static TArray<TSharedRef<class IClassViewerFilter>> NotImplemented;
+	return NotImplemented;
+}
+
 void FPropertyRowGenerator::PreSetObject(int32 NumNewObjects, bool bHasStructRoots)
 {
 	// Save existing expanded items first
diff --git a/Engine/Source/Editor/PropertyEditor/Private/PropertyRowGenerator.h b/Engine/Source/Editor/PropertyEditor/Private/PropertyRowGenerator.h
--- a/Engine/Source/Editor/PropertyEditor/Private/PropertyRowGenerator.h
+++ b/Engine/Source/Editor/PropertyEditor/Private/PropertyRowGenerator.h
@@ -110,6 +110,7 @@ public:
 	virtual void ForceRefresh();
 	virtual TSharedPtr<class FAssetThumbnailPool> GetThumbnailPool() const;
 	virtual bool HasClassDefaultObject() const { return bViewingClassDefaultObject; }
+	virtual const TArray<TSharedRef<class IClassViewerFilter>>& GetClassViewerFilters() const;
 
 	const FCustomPropertyTypeLayoutMap& GetInstancedPropertyTypeLayoutMap() const;
 	void UpdateDetailRows();
diff --git a/Engine/Source/Editor/PropertyEditor/Private/SDetailsView.cpp b/Engine/Source/Editor/PropertyEditor/Private/SDetailsView.cpp
--- a/Engine/Source/Editor/PropertyEditor/Private/SDetailsView.cpp
+++ b/Engine/Source/Editor/PropertyEditor/Private/SDetailsView.cpp
@@ -50,6 +50,7 @@ void SDetailsView::Construct(const FArguments& InArgs, const FDetailsViewArgs& I
 	ColumnSizeData.RightColumnMinWidth = DetailsViewArgs.RightColumnMinWidth;
 
 	SetObjectFilter(InDetailsViewArgs.ObjectFilter);
+	SetClassViewerFilters(InDetailsViewArgs.ClassViewerFilters);
 
 	bViewingClassDefaultObject = false;
 
@@ -520,6 +521,11 @@ void SDetailsView::SetObjectFilter(TSharedPtr<FDetailsViewObjectFilter> InFilter
 	}
 }
 
+void SDetailsView::SetClassViewerFilters(const TArray<TSharedRef<class IClassViewerFilter>>& InFilters)
+{
+	ClassViewerFilters = InFilters;
+}
+
 bool SDetailsView::ShouldSetNewObjects(const TArray<UObject*>& InObjects) const
 {
 	bool bShouldSetObjects = false;
diff --git a/Engine/Source/Editor/PropertyEditor/Private/SDetailsView.h b/Engine/Source/Editor/PropertyEditor/Private/SDetailsView.h
--- a/Engine/Source/Editor/PropertyEditor/Private/SDetailsView.h
+++ b/Engine/Source/Editor/PropertyEditor/Private/SDetailsView.h
@@ -52,6 +52,7 @@ public:
 	virtual void SetRootObjectCustomizationInstance(TSharedPtr<IDetailRootObjectCustomization> InRootObjectCustomization) override;
 	virtual void ClearSearch() override;
 	virtual void SetObjectFilter(TSharedPtr<FDetailsViewObjectFilter> InFilter) override;
+	virtual void SetClassViewerFilters(const TArray<TSharedRef<class IClassViewerFilter>>& InFilters) override;
 
 	/**
 	 * Replaces objects being observed by the view with new objects
diff --git a/Engine/Source/Editor/PropertyEditor/Private/SDetailsViewBase.cpp b/Engine/Source/Editor/PropertyEditor/Private/SDetailsViewBase.cpp
--- a/Engine/Source/Editor/PropertyEditor/Private/SDetailsViewBase.cpp
+++ b/Engine/Source/Editor/PropertyEditor/Private/SDetailsViewBase.cpp
@@ -544,6 +544,11 @@ TSharedPtr<FAssetThumbnailPool> SDetailsViewBase::GetThumbnailPool() const
 	return UThumbnailManager::Get().GetSharedThumbnailPool();
 }
 
+const TArray<TSharedRef<class IClassViewerFilter>>& SDetailsViewBase::GetClassViewerFilters() const
+{
+	return ClassViewerFilters;
+}
+
 void SDetailsViewBase::NotifyFinishedChangingProperties(const FPropertyChangedEvent& PropertyChangedEvent)
 {
 	OnFinishedChangingPropertiesDelegate.Broadcast(PropertyChangedEvent);
diff --git a/Engine/Source/Editor/PropertyEditor/Private/SDetailsViewBase.h b/Engine/Source/Editor/PropertyEditor/Private/SDetailsViewBase.h
--- a/Engine/Source/Editor/PropertyEditor/Private/SDetailsViewBase.h
+++ b/Engine/Source/Editor/PropertyEditor/Private/SDetailsViewBase.h
@@ -161,6 +161,7 @@ public:
 	virtual void NotifyFinishedChangingProperties(const FPropertyChangedEvent& PropertyChangedEvent) override;
 	void RefreshTree() override;
 	TSharedPtr<class FAssetThumbnailPool> GetThumbnailPool() const override;
+	virtual const TArray<TSharedRef<class IClassViewerFilter>>& GetClassViewerFilters() const override;
 	TSharedPtr<IPropertyUtilities> GetPropertyUtilities() override;
 	void CreateColorPickerWindow(const TSharedRef< class FPropertyEditor >& PropertyEditor, bool bUseAlpha) override;
 	virtual void UpdateSinglePropertyMap(TSharedPtr<FComplexPropertyNode> InRootPropertyNode, FDetailLayoutData& LayoutData, bool bIsExternal) override;
@@ -454,4 +455,6 @@ protected:
 
 	mutable TSharedPtr<FEditConditionParser> EditConditionParser;
 
+	/** Optional custom filter(s) to be applied when selecting values for class properties */
+	TArray<TSharedRef<class IClassViewerFilter>> ClassViewerFilters;
 };
diff --git a/Engine/Source/Editor/PropertyEditor/Private/SPropertyTreeViewImpl.cpp b/Engine/Source/Editor/PropertyEditor/Private/SPropertyTreeViewImpl.cpp
--- a/Engine/Source/Editor/PropertyEditor/Private/SPropertyTreeViewImpl.cpp
+++ b/Engine/Source/Editor/PropertyEditor/Private/SPropertyTreeViewImpl.cpp
@@ -80,6 +80,13 @@ public:
 		return NULL;
 	}
 
+	virtual const TArray<TSharedRef<class IClassViewerFilter>>& GetClassViewerFilters() const override
+	{
+		// not implemented
+		static TArray<TSharedRef<class IClassViewerFilter>> NotImplemented;
+		return NotImplemented;
+	}
+
 	virtual void NotifyFinishedChangingProperties(const FPropertyChangedEvent& PropertyChangedEvent) override {}
 
 	virtual bool DontUpdateValueWhileEditing() const override
diff --git a/Engine/Source/Editor/PropertyEditor/Private/SSingleProperty.cpp b/Engine/Source/Editor/PropertyEditor/Private/SSingleProperty.cpp
--- a/Engine/Source/Editor/PropertyEditor/Private/SSingleProperty.cpp
+++ b/Engine/Source/Editor/PropertyEditor/Private/SSingleProperty.cpp
@@ -79,6 +79,13 @@ public:
 		return false;
 	}
 
+	virtual const TArray<TSharedRef<class IClassViewerFilter>>& GetClassViewerFilters() const override
+	{
+		// not implemented
+		static TArray<TSharedRef<class IClassViewerFilter>> NotImplemented;
+		return NotImplemented;
+	}
+
 	const TArray<TWeakObjectPtr<UObject>>& GetSelectedObjects() const override
 	{
 		static TArray<TWeakObjectPtr<UObject>> Empty;
diff --git a/Engine/Source/Editor/PropertyEditor/Private/SStructureDetailsView.h b/Engine/Source/Editor/PropertyEditor/Private/SStructureDetailsView.h
--- a/Engine/Source/Editor/PropertyEditor/Private/SStructureDetailsView.h
+++ b/Engine/Source/Editor/PropertyEditor/Private/SStructureDetailsView.h
@@ -95,6 +95,7 @@ public:
 	virtual void SetRootObjectCustomizationInstance(TSharedPtr<IDetailRootObjectCustomization> InRootObjectCustomization) override {}
 	virtual TSharedPtr<class IDetailRootObjectCustomization> GetRootObjectCustomization() const override { return nullptr; }
 	virtual void SetObjectFilter(TSharedPtr<FDetailsViewObjectFilter> InFilter) override {}
+	virtual void SetClassViewerFilters(const TArray<TSharedRef<class IClassViewerFilter>>& InFilters) override {}
 	virtual bool IsGroupFavorite(FStringView GroupPath) const { return false; }
 	virtual void SetGroupFavorite(FStringView GroupPath, bool IsFavorite) {}
 	virtual bool IsCustomBuilderFavorite(FStringView Path) const { return false; }
diff --git a/Engine/Source/Editor/PropertyEditor/Private/UserInterface/PropertyDetails/PropertyDetailsUtilities.cpp b/Engine/Source/Editor/PropertyEditor/Private/UserInterface/PropertyDetails/PropertyDetailsUtilities.cpp
--- a/Engine/Source/Editor/PropertyEditor/Private/UserInterface/PropertyDetails/PropertyDetailsUtilities.cpp
+++ b/Engine/Source/Editor/PropertyEditor/Private/UserInterface/PropertyDetails/PropertyDetailsUtilities.cpp
@@ -55,6 +55,11 @@ TSharedPtr<class FAssetThumbnailPool> FPropertyDetailsUtilities::GetThumbnailPoo
 	return DetailsView.GetThumbnailPool();
 }
 
+const TArray<TSharedRef<class IClassViewerFilter>>& FPropertyDetailsUtilities::GetClassViewerFilters() const
+{
+	return DetailsView.GetClassViewerFilters();
+}
+
 void FPropertyDetailsUtilities::NotifyFinishedChangingProperties(const FPropertyChangedEvent& PropertyChangedEvent)
 {
 	DetailsView.NotifyFinishedChangingProperties(PropertyChangedEvent);
diff --git a/Engine/Source/Editor/PropertyEditor/Private/UserInterface/PropertyDetails/PropertyDetailsUtilities.h b/Engine/Source/Editor/PropertyEditor/Private/UserInterface/PropertyDetails/PropertyDetailsUtilities.h
--- a/Engine/Source/Editor/PropertyEditor/Private/UserInterface/PropertyDetails/PropertyDetailsUtilities.h
+++ b/Engine/Source/Editor/PropertyEditor/Private/UserInterface/PropertyDetails/PropertyDetailsUtilities.h
@@ -26,6 +26,7 @@ public:
 	virtual bool DontUpdateValueWhileEditing() const override;
 	const TArray<TWeakObjectPtr<UObject>>& GetSelectedObjects() const override;
 	virtual bool HasClassDefaultObject() const override;
+	virtual const TArray<TSharedRef<class IClassViewerFilter>>& GetClassViewerFilters() const override;
 
 private:
 	IDetailsViewPrivate& DetailsView;
diff --git a/Engine/Source/Editor/PropertyEditor/Private/UserInterface/PropertyEditor/SPropertyEditorClass.cpp b/Engine/Source/Editor/PropertyEditor/Private/UserInterface/PropertyEditor/SPropertyEditorClass.cpp
--- a/Engine/Source/Editor/PropertyEditor/Private/UserInterface/PropertyEditor/SPropertyEditorClass.cpp
+++ b/Engine/Source/Editor/PropertyEditor/Private/UserInterface/PropertyEditor/SPropertyEditorClass.cpp
@@ -194,7 +194,7 @@ void SPropertyEditorClass::Construct(const FArguments& InArgs, const TSharedPtr<
 		OnSetClass = InArgs._OnSetClass;
 	}
 
-	CreateClassFilter();
+	CreateClassFilter(InArgs._ClassViewerFilters);
 
 	SAssignNew(ComboButton, SComboButton)
 		.OnGetMenuContent(this, &SPropertyEditorClass::GenerateClassPicker)
@@ -261,7 +261,7 @@ FText SPropertyEditorClass::GetDisplayValueAsString() const
 	
 }
 
-void SPropertyEditorClass::CreateClassFilter()
+void SPropertyEditorClass::CreateClassFilter(const TArray<TSharedRef<IClassViewerFilter>>& InClassFilters)
 {
 	ClassViewerOptions.bShowBackgroundBorder = false;
 	ClassViewerOptions.bShowUnloadedBlueprints = true;
@@ -277,9 +277,10 @@ void SPropertyEditorClass::CreateClassFilter()
 	ClassViewerOptions.NameTypeToDisplay = (bShowDisplayNames ? EClassViewerNameTypeToDisplay::DisplayName : EClassViewerNameTypeToDisplay::ClassName);
 	ClassViewerOptions.DisplayMode = bShowTree ? EClassViewerDisplayMode::TreeView : EClassViewerDisplayMode::ListView;
 	ClassViewerOptions.bAllowViewOptions = bShowViewOptions;
+	ClassViewerOptions.ClassFilters.Append(InClassFilters);
 
-	TSharedPtr<FPropertyEditorClassFilter> PropEdClassFilter = MakeShareable(new FPropertyEditorClassFilter);
-	ClassViewerOptions.ClassFilter = PropEdClassFilter;
+	TSharedRef<FPropertyEditorClassFilter> PropEdClassFilter = MakeShared<FPropertyEditorClassFilter>();
+	ClassViewerOptions.ClassFilters.Add(PropEdClassFilter);
 
 	PropEdClassFilter->ClassPropertyMetaClass = MetaClass;
 	PropEdClassFilter->InterfaceThatMustBeImplemented = RequiredInterface;
diff --git a/Engine/Source/Editor/PropertyEditor/Private/UserInterface/PropertyEditor/SPropertyEditorClass.h b/Engine/Source/Editor/PropertyEditor/Private/UserInterface/PropertyEditor/SPropertyEditorClass.h
--- a/Engine/Source/Editor/PropertyEditor/Private/UserInterface/PropertyEditor/SPropertyEditorClass.h
+++ b/Engine/Source/Editor/PropertyEditor/Private/UserInterface/PropertyEditor/SPropertyEditorClass.h
@@ -58,6 +58,8 @@ public:
 		SLATE_ARGUMENT(bool, ShowDisplayNames)
 		/** Delegate used to set the currently selected class (required if PropertyEditor == null) */
 		SLATE_EVENT(FOnSetClass, OnSetClass)
+		/** Custom class filter(s) to be applied to the class picker widget (may be empty) */
+		SLATE_ARGUMENT(TArray<TSharedRef<class IClassViewerFilter>>, ClassViewerFilters)
 	SLATE_END_ARGS()
 
 	static bool Supports(const TSharedRef< class FPropertyEditor >& InPropertyEditor);
@@ -139,5 +141,5 @@ private:
 	/** Delegate used to set the currently selected class (required if PropertyEditor == null) */
 	FOnSetClass OnSetClass;
 
-	void CreateClassFilter();
+	void CreateClassFilter(const TArray<TSharedRef<IClassViewerFilter>>& InClassFilters);
 };
diff --git a/Engine/Source/Editor/PropertyEditor/Private/UserInterface/PropertyEditor/SPropertyEditorEditInline.cpp b/Engine/Source/Editor/PropertyEditor/Private/UserInterface/PropertyEditor/SPropertyEditorEditInline.cpp
--- a/Engine/Source/Editor/PropertyEditor/Private/UserInterface/PropertyEditor/SPropertyEditorEditInline.cpp
+++ b/Engine/Source/Editor/PropertyEditor/Private/UserInterface/PropertyEditor/SPropertyEditorEditInline.cpp
@@ -171,7 +171,7 @@ TSharedRef<SWidget> SPropertyEditorEditInline::GenerateClassPicker()
 	Options.NameTypeToDisplay = EClassViewerNameTypeToDisplay::DisplayName;
 
 	TSharedPtr<FPropertyEditorInlineClassFilter> ClassFilter = MakeShareable( new FPropertyEditorInlineClassFilter );
-	Options.ClassFilter = ClassFilter;
+	Options.ClassFilters.Add(ClassFilter.ToSharedRef());
 	ClassFilter->bAllowAbstract = false;
 
 	const TSharedRef< FPropertyNode > PropertyNode = PropertyEditor->GetPropertyNode();
diff --git a/Engine/Source/Editor/PropertyEditor/Public/IDetailsView.h b/Engine/Source/Editor/PropertyEditor/Public/IDetailsView.h
--- a/Engine/Source/Editor/PropertyEditor/Public/IDetailsView.h
+++ b/Engine/Source/Editor/PropertyEditor/Public/IDetailsView.h
@@ -53,6 +53,8 @@ struct FDetailsViewArgs
 	TSharedPtr<FTabManager> HostTabManager;
 	/** Optional object filter to use for more complex handling of what a details panel is viewing. */
 	TSharedPtr<FDetailsViewObjectFilter> ObjectFilter;
+	/** Optional custom filter(s) to apply to the class viewer widget for class object property values. */
+	TArray<TSharedRef<class IClassViewerFilter>> ClassViewerFilters;
 
 	/** Identifier for this details view; NAME_None if this view is anonymous */
 	FName ViewIdentifier;
@@ -365,7 +367,7 @@ public:
 	/** Returns the name area widget used to display object naming functionality so it can be placed in a custom location.  Note FDetailsViewArgs.bCustomNameAreaLocation must be true */
 	virtual TSharedPtr<SWidget> GetNameAreaWidget() = 0;
 
-	/** Optinoally add custom tools into the NameArea */
+	/** Optionally add custom tools into the NameArea */
 	virtual void SetNameAreaCustomContent(TSharedRef<SWidget>& InCustomContent) = 0;
 
 	/** Returns the search area widget used to display search and view options so it can be placed in a custom location.  Note FDetailsViewArgs.bCustomFilterAreaLocation must be true */
@@ -386,6 +388,9 @@ public:
 	/** Sets an optional object filter to use for more complex handling of what a details panel is viewing. */
 	virtual void SetObjectFilter(TSharedPtr<FDetailsViewObjectFilter> InFilter) = 0;
 
+	/** Sets the custom filter(s) to be used when selecting values for class properties in this view. */
+	virtual void SetClassViewerFilters(const TArray<TSharedRef<class IClassViewerFilter>>& InFilters) = 0;
+
 	/** Allows other systems to add a custom filter in the details panel */
 	virtual void SetCustomFilterDelegate(FSimpleDelegate InDelegate) = 0;
 
diff --git a/Engine/Source/Editor/PropertyEditor/Public/IPropertyUtilities.h b/Engine/Source/Editor/PropertyEditor/Public/IPropertyUtilities.h
--- a/Engine/Source/Editor/PropertyEditor/Public/IPropertyUtilities.h
+++ b/Engine/Source/Editor/PropertyEditor/Public/IPropertyUtilities.h
@@ -29,4 +29,5 @@ public:
 	virtual const TArray<TWeakObjectPtr<UObject>>& GetSelectedObjects() const = 0;
 	/** If a customization standalone widget is used, the value should be update only once, when its window is closed */
 	virtual bool DontUpdateValueWhileEditing() const = 0;
+	virtual const TArray<TSharedRef<class IClassViewerFilter>>& GetClassViewerFilters() const = 0;
 };
diff --git a/Engine/Source/Editor/Sequencer/Private/SequencerContextMenus.cpp b/Engine/Source/Editor/Sequencer/Private/SequencerContextMenus.cpp
--- a/Engine/Source/Editor/Sequencer/Private/SequencerContextMenus.cpp
+++ b/Engine/Source/Editor/Sequencer/Private/SequencerContextMenus.cpp
@@ -2202,7 +2202,7 @@ void FEasingContextMenu::EasingTypeMenu(FMenuBuilder& MenuBuilder)
 
 	FClassViewerInitializationOptions InitOptions;
 	InitOptions.NameTypeToDisplay = EClassViewerNameTypeToDisplay::DisplayName;
-	InitOptions.ClassFilter = MakeShared<FFilter>();
+	InitOptions.ClassFilters.Add(MakeShared<FFilter>());
 
 	// Copy a reference to the context menu by value into each lambda handler to ensure the type stays alive until the menu is closed
 	TSharedRef<FEasingContextMenu> Shared = AsShared();
diff --git a/Engine/Source/Editor/UMGEditor/Private/WidgetBlueprintFactory.cpp b/Engine/Source/Editor/UMGEditor/Private/WidgetBlueprintFactory.cpp
--- a/Engine/Source/Editor/UMGEditor/Private/WidgetBlueprintFactory.cpp
+++ b/Engine/Source/Editor/UMGEditor/Private/WidgetBlueprintFactory.cpp
@@ -75,7 +75,7 @@ bool UWidgetBlueprintFactory::ConfigureProperties()
 		Options.ExtraPickerCommonClasses.Add(UCanvasPanel::StaticClass());
 
 		TSharedPtr<FWidgetClassFilter> Filter = MakeShareable(new FWidgetClassFilter);
-		Options.ClassFilter = Filter;
+		Options.ClassFilters.Add(Filter.ToSharedRef());
 
 		Filter->DisallowedClassFlags = CLASS_Abstract | CLASS_Deprecated | CLASS_NewerVersionExists;
 		Filter->AllowedChildrenOfClasses.Add(UPanelWidget::StaticClass());
diff --git a/Engine/Source/Editor/UnrealEd/Private/Factories/AnimBlueprintFactory.cpp b/Engine/Source/Editor/UnrealEd/Private/Factories/AnimBlueprintFactory.cpp
--- a/Engine/Source/Editor/UnrealEd/Private/Factories/AnimBlueprintFactory.cpp
+++ b/Engine/Source/Editor/UnrealEd/Private/Factories/AnimBlueprintFactory.cpp
@@ -200,7 +200,7 @@ private:
 		Options.bIsBlueprintBaseOnly = true;
 
 		TSharedPtr<FAnimBlueprintParentFilter> Filter = MakeShareable(new FAnimBlueprintParentFilter(TargetSkeleton));
-		Options.ClassFilter = Filter;
+		Options.ClassFilters.Add(Filter.ToSharedRef());
 
 		// All child child classes of UAnimInstance are valid.
 		Filter->AllowedChildrenOfClasses.Add(UAnimInstance::StaticClass());
diff --git a/Engine/Source/Editor/UnrealEd/Private/Factories/EditorFactories.cpp b/Engine/Source/Editor/UnrealEd/Private/Factories/EditorFactories.cpp
--- a/Engine/Source/Editor/UnrealEd/Private/Factories/EditorFactories.cpp
+++ b/Engine/Source/Editor/UnrealEd/Private/Factories/EditorFactories.cpp
@@ -1974,7 +1974,7 @@ bool UPhysicalMaterialFactoryNew::ConfigureProperties()
 	Options.Mode = EClassViewerMode::ClassPicker;
 
 	TSharedPtr<FAssetClassParentFilter> Filter = MakeShareable(new FAssetClassParentFilter);
-	Options.ClassFilter = Filter;
+	Options.ClassFilters.Add(Filter.ToSharedRef());
 
 	Filter->DisallowedClassFlags = CLASS_Abstract | CLASS_Deprecated | CLASS_NewerVersionExists;
 	Filter->AllowedChildrenOfClasses.Add(UPhysicalMaterial::StaticClass());
@@ -2033,7 +2033,7 @@ bool UPhysicalMaterialMaskFactory::ConfigureProperties()
 	Options.Mode = EClassViewerMode::ClassPicker;
 
 	TSharedPtr<FAssetClassParentFilter> Filter = MakeShareable(new FAssetClassParentFilter);
-	Options.ClassFilter = Filter;
+	Options.ClassFilters.Add(Filter.ToSharedRef());
 
 	Filter->DisallowedClassFlags = CLASS_Abstract | CLASS_Deprecated | CLASS_NewerVersionExists;
 	Filter->AllowedChildrenOfClasses.Add(UPhysicalMaterialMask::StaticClass());
@@ -6743,7 +6743,7 @@ bool UBlueprintFactory::ConfigureProperties()
 
 	// Prevent creating blueprints of classes that require special setup (they'll be allowed in the corresponding factories / via other means)
 	TSharedPtr<FBlueprintParentFilter> Filter = MakeShareable(new FBlueprintParentFilter);
-	Options.ClassFilter = Filter;
+	Options.ClassFilters.Add(Filter.ToSharedRef());
 	if (!IsMacroFactory())
 	{
 		Filter->DisallowedChildrenOfClasses.Add(ALevelScriptActor::StaticClass());
@@ -7023,7 +7023,7 @@ bool UCurveFactory::ConfigureProperties()
 	Options.Mode = EClassViewerMode::ClassPicker;
 
 	TSharedPtr<FAssetClassParentFilter> Filter = MakeShareable(new FAssetClassParentFilter);
-	Options.ClassFilter = Filter;
+	Options.ClassFilters.Add(Filter.ToSharedRef());
 
 	Filter->DisallowedClassFlags = CLASS_Abstract | CLASS_Deprecated | CLASS_NewerVersionExists;
 	Filter->AllowedChildrenOfClasses.Add(UCurveBase::StaticClass());
@@ -7242,7 +7242,7 @@ bool UDataAssetFactory::ConfigureProperties()
 	Options.Mode = EClassViewerMode::ClassPicker;
 
 	TSharedPtr<FAssetClassParentFilter> Filter = MakeShareable(new FAssetClassParentFilter);
-	Options.ClassFilter = Filter;
+	Options.ClassFilters.Add(Filter.ToSharedRef());
 
 	Filter->DisallowedClassFlags = CLASS_Abstract | CLASS_Deprecated | CLASS_NewerVersionExists | CLASS_HideDropDown;
 	Filter->AllowedChildrenOfClasses.Add(UDataAsset::StaticClass());
diff --git a/Engine/Source/Editor/UnrealEd/Private/Factories/LightWeightInstanceFactory.cpp b/Engine/Source/Editor/UnrealEd/Private/Factories/LightWeightInstanceFactory.cpp
--- a/Engine/Source/Editor/UnrealEd/Private/Factories/LightWeightInstanceFactory.cpp
+++ b/Engine/Source/Editor/UnrealEd/Private/Factories/LightWeightInstanceFactory.cpp
@@ -67,7 +67,7 @@ bool ULightWeightInstanceFactory::ConfigureProperties()
 
 	// Prevent creating blueprints of classes that require special setup (they'll be allowed in the corresponding factories / via other means)
 	TSharedPtr<FLWIFilter> Filter = MakeShareable(new FLWIFilter);
-	Options.ClassFilter = Filter;
+	Options.ClassFilters.Add(Filter.ToSharedRef());
 
 	Filter->AllowedChildrenOfClasses.Add(ALightWeightInstanceManager::StaticClass());
 
diff --git a/Engine/Source/Editor/UnrealEd/Private/Factories/SlateWidgetStyleAssetFactory.cpp b/Engine/Source/Editor/UnrealEd/Private/Factories/SlateWidgetStyleAssetFactory.cpp
--- a/Engine/Source/Editor/UnrealEd/Private/Factories/SlateWidgetStyleAssetFactory.cpp
+++ b/Engine/Source/Editor/UnrealEd/Private/Factories/SlateWidgetStyleAssetFactory.cpp
@@ -61,7 +61,7 @@ bool USlateWidgetStyleAssetFactory::ConfigureProperties()
 	Options.Mode = EClassViewerMode::ClassPicker;
 
 	TSharedPtr<FClassFilter> Filter = MakeShareable(new FClassFilter);
-	Options.ClassFilter = Filter;
+	Options.ClassFilters.Add(Filter.ToSharedRef());
 
 	Filter->DisallowedClassFlags = CLASS_Abstract | CLASS_Deprecated | CLASS_NewerVersionExists;
 	Filter->AllowedChildrenOfClasses.Add( USlateWidgetStyleContainerBase::StaticClass() );
diff --git a/Engine/Source/Editor/UnrealEd/Private/Kismet2/BlueprintEditorUtils.cpp b/Engine/Source/Editor/UnrealEd/Private/Kismet2/BlueprintEditorUtils.cpp
--- a/Engine/Source/Editor/UnrealEd/Private/Kismet2/BlueprintEditorUtils.cpp
+++ b/Engine/Source/Editor/UnrealEd/Private/Kismet2/BlueprintEditorUtils.cpp
@@ -8358,7 +8358,7 @@ TSharedRef<SWidget> FBlueprintEditorUtils::ConstructBlueprintParentClassPicker(
 	Options.bShowBackgroundBorder = false;
 
 	TSharedPtr<FBlueprintReparentFilter> Filter = MakeShareable(new FBlueprintReparentFilter);
-	Options.ClassFilter = Filter;
+	Options.ClassFilters.Add(Filter.ToSharedRef());
 	Options.ViewerTitleString = LOCTEXT("ReparentBlueprint", "Reparent blueprint");
 
 	// Only allow parenting to base blueprints.
@@ -8519,7 +8519,7 @@ TSharedRef<SWidget> FBlueprintEditorUtils::ConstructBlueprintInterfaceClassPicke
 	Options.bShowBackgroundBorder = false;
 
 	TSharedPtr<FBlueprintInterfaceFilter> Filter = MakeShareable(new FBlueprintInterfaceFilter);
-	Options.ClassFilter = Filter;
+	Options.ClassFilters.Add(Filter.ToSharedRef());
 	Options.ViewerTitleString = LOCTEXT("ImplementInterfaceBlueprint", "Implement Interface");
 
 	for (const UBlueprint* Blueprint : Blueprints)
