Fix for non-unity build.

#rnx

diff --git a/Engine/Source/Editor/ClassViewer/Public/ClassViewerFilter.h b/Engine/Source/Editor/ClassViewer/Public/ClassViewerFilter.h
--- a/Engine/Source/Editor/ClassViewer/Public/ClassViewerFilter.h
+++ b/Engine/Source/Editor/ClassViewer/Public/ClassViewerFilter.h
@@ -4,6 +4,7 @@
 #include "CoreMinimal.h"
 #include "Engine/EngineTypes.h"
 #include "Settings/ClassViewerSettings.h"
+#include "Misc/Attribute.h"
 
 class FClassViewerNode;
 class FClassViewerInitializationOptions;
