SClassViewer: Skip hidden classes

Opening the class viewer widget was causing a crash for hidden classes created by [REDACTED_PERSON].

diff --git a/Engine/Source/Editor/ClassViewer/Private/SClassViewer.cpp b/Engine/Source/Editor/ClassViewer/Private/SClassViewer.cpp
--- a/Engine/Source/Editor/ClassViewer/Private/SClassViewer.cpp
+++ b/Engine/Source/Editor/ClassViewer/Private/SClassViewer.cpp
@@ -1071,7 +1071,7 @@ void FClassHierarchy::AddChildren_NoFilter( TSharedPtr< FClassViewerNode >& InOu
 		UClass* CurrentClass = *ClassIt;
 
 		// Ignore deprecated and temporary trash classes.
-		if (CurrentClass->HasAnyClassFlags(CLASS_Deprecated | CLASS_NewerVersionExists) ||
+		if (CurrentClass->HasAnyClassFlags(CLASS_Deprecated | CLASS_NewerVersionExists | CLASS_Hidden) ||
 			FKismetEditorUtilities::IsClassABlueprintSkeleton(CurrentClass))
 		{
 			continue;
