FClassHierarchy::PopulateClassHierarcy: Refactor for robustness.
Instead of creating a tree for classes-from-blueprint and a tree for classes-from-loaded-UClass, create a single map from ClassPathName to node, and populate the data-from-blueprint and data-from-loaded-Uclass on the same node, and then create a tree from that map.
#rnx

diff --git a/Engine/Source/Editor/ClassViewer/Private/ClassViewerNode.cpp b/Engine/Source/Editor/ClassViewer/Private/ClassViewerNode.cpp
--- a/Engine/Source/Editor/ClassViewer/Private/ClassViewerNode.cpp
+++ b/Engine/Source/Editor/ClassViewer/Private/ClassViewerNode.cpp
@@ -76,39 +76,6 @@ void FClassViewerNode::AddChild( TSharedPtr<FClassViewerNode> Child )
 	ChildrenList.Add(Child);
 }
 
-void FClassViewerNode::AddUniqueChild(TSharedPtr<FClassViewerNode> NewChild)
-{
-	check(NewChild.IsValid());
-	const UClass* NewChildClass = NewChild->Class.Get();
-	if (nullptr != NewChildClass)
-	{
-		for(int ChildIndex = 0; ChildIndex < ChildrenList.Num(); ++ChildIndex)
-		{
-			TSharedPtr<FClassViewerNode> OldChild = ChildrenList[ChildIndex];
-			if(OldChild.IsValid() && OldChild->Class == NewChildClass)
-			{
-				const bool bNewChildHasMoreInfo = NewChild->UnloadedBlueprintData.IsValid();
-				const bool bOldChildHasMoreInfo = OldChild->UnloadedBlueprintData.IsValid();
-				if(bNewChildHasMoreInfo && !bOldChildHasMoreInfo)
-				{
-					// make sure, that new child has all needed children
-					for(int OldChildIndex = 0; OldChildIndex < OldChild->ChildrenList.Num(); ++OldChildIndex)
-					{
-						NewChild->AddUniqueChild( OldChild->ChildrenList[OldChildIndex] );
-					}
-
-					// replace child
-					NewChild->ParentNode = AsShared();
-					ChildrenList[ChildIndex] = NewChild;
-				}
-				return;
-			}
-		}
-	}
-
-	AddChild(NewChild);
-}
-
 bool FClassViewerNode::IsRestricted() const
 {
 	return PropertyHandle.IsValid() && PropertyHandle->IsRestricted(*ClassName);
@@ -181,4 +148,4 @@ bool FClassViewerNode::IsEditorOnlyClass() const
 TSharedPtr< FClassViewerNode > FClassViewerNode::GetParentNode() const
 {
 	return ParentNode.Pin();
-}
\ No newline at end of file
+}
diff --git a/Engine/Source/Editor/ClassViewer/Private/ClassViewerNode.h b/Engine/Source/Editor/ClassViewer/Private/ClassViewerNode.h
--- a/Engine/Source/Editor/ClassViewer/Private/ClassViewerNode.h
+++ b/Engine/Source/Editor/ClassViewer/Private/ClassViewerNode.h
@@ -33,14 +33,6 @@ public:
 	 */
 	void AddChild( TSharedPtr<FClassViewerNode> Child );
 
-	/**
-	 * Adds the specified child to the node. If a child with the same class already exists the function add the child storing more info.
-	 * The function does not persist child order.
-	 *
-	 * @param	NewChild	The child to be added to this node for the tree.
-	 */
-	void AddUniqueChild(TSharedPtr<FClassViewerNode> NewChild);
-
 	/** 
 	 * Retrieves the class name this node is associated with. This is not the literal UClass name as it is missing the _C for blueprints
 	 * @param	bUseDisplayName	Whether to use the display name or class name
diff --git a/Engine/Source/Editor/ClassViewer/Private/SClassViewer.cpp b/Engine/Source/Editor/ClassViewer/Private/SClassViewer.cpp
--- a/Engine/Source/Editor/ClassViewer/Private/SClassViewer.cpp
+++ b/Engine/Source/Editor/ClassViewer/Private/SClassViewer.cpp
@@ -151,11 +151,11 @@ public:
 	TSharedPtr< FClassViewerNode > FindNodeByGeneratedClassPath(const TSharedPtr< FClassViewerNode >& InRootNode, FName InGeneratedClassPath);
 
 private:
-	/** Recursive function to build a tree, will not filter.
-	 *	@param InOutRootNode						The node that this function will add the children of to the tree.
-	 *  @param InOutClassPathToNode			        Map that will be populated with the class path to the corresponding class node.
+	/** Adds UClass information to The node in InOutClassPathToNode for every loaded class, creating the node if it does not exist. Does not filter classes.
+	 *	@param OutRootNode						Out parameter that receives the node pointer for the root UClass.
+	 *  @param InOutClassPathToNode				Map containing all nodes.
 	 */	
-	void AddChildren_NoFilter( TSharedPtr< FClassViewerNode >& InOutRootNode, TMap<FName, TSharedPtr< FClassViewerNode >>& InOutClassPathToNode );
+	void CreateNodesForLoadedClasses( TSharedPtr<FClassViewerNode>& OutRootNode, TMap<FName, TSharedPtr< FClassViewerNode >>& InOutClassPathToNode );
 
 	/** Called when reload has finished */
 	void OnReloadComplete( EReloadCompleteReason Reason );
@@ -168,6 +168,14 @@ private:
 	 */
 	void LoadUnloadedTagData(TSharedPtr<FClassViewerNode>& InOutClassViewerNode, const FAssetData& InAssetData);
 
+	/**
+	 * Sets the fields calculated from AssetData on the given Node, if not already set
+	 *
+	 * @param InOutClassViewerNode		The node to save all the data into.
+	 * @param InAssetData				The asset data to pull the tags from.
+	 */
+	void SetAssetDataFields(TSharedPtr<FClassViewerNode>& InOutClassViewerNode, const FAssetData& InAssetData);
+
 	/**
 	 * Finds the UClass and UBlueprint for the passed in node, utilizing unloaded data to find it.
 	 *
@@ -175,6 +183,9 @@ private:
 	 */
 	void FindClass(TSharedPtr< FClassViewerNode > InOutClassNode);
 
+	/** Sets the fields calculated from the in-memory UClass on the given Node, if not already set. */
+	void SetClassFields(TSharedPtr<FClassViewerNode>& InOutClassNode, UClass& Class);
+
 	/** 
 	 * Recursively searches through the hierarchy to find and remove the asset. Used when deleting assets.
 	 *
@@ -1050,75 +1061,27 @@ void FClassHierarchy::OnReloadComplete(EReloadCompleteReason Reason)
 	ClassViewer::Helpers::RequestPopulateClassHierarchy();
 }
 
-void FClassHierarchy::AddChildren_NoFilter( TSharedPtr< FClassViewerNode >& InOutRootNode, TMap<FName, TSharedPtr< FClassViewerNode >>& InOutClassPathToNode )
+void FClassHierarchy::CreateNodesForLoadedClasses(TSharedPtr<FClassViewerNode>& OutRootNode, TMap<FName, TSharedPtr< FClassViewerNode >>& InOutClassPathToNode)
 {
-	UClass* RootClass = UObject::StaticClass();
-
-	ObjectClassRoot = MakeShared<FClassViewerNode>(RootClass);
-	ObjectClassRoot->Class = RootClass;
-
-	InOutClassPathToNode.Add(ObjectClassRoot->ClassPath, ObjectClassRoot);
-
-	TMap< UClass*, TSharedPtr< FClassViewerNode > > VisitedNodes;
-	VisitedNodes.Add(RootClass, ObjectClassRoot);
-
-	// Go through all of the classes children and see if they should be added to the list.
-
-	TArray<TPair<UClass*, TSharedPtr<FClassViewerNode>>> EntriesToProcess;
-
 	for (TObjectIterator<UClass> ClassIt; ClassIt; ++ClassIt)
 	{
 		UClass* CurrentClass = *ClassIt;
-
 		// Ignore deprecated and temporary trash classes.
 		if (CurrentClass->HasAnyClassFlags(CLASS_Deprecated | CLASS_NewerVersionExists | CLASS_Hidden) ||
 			FKismetEditorUtilities::IsClassABlueprintSkeleton(CurrentClass))
 		{
 			continue;
 		}
-		
-		TSharedPtr<FClassViewerNode>& Entry = VisitedNodes.FindOrAdd(CurrentClass);
-		// If the entry is valid we already visited it.
-		if (Entry.IsValid())
+		TSharedPtr<FClassViewerNode>& Node = InOutClassPathToNode.FindOrAdd(FName(*CurrentClass->GetPathName()));
+		if (!Node)
 		{
-			continue;
-		}
-		else
-		{
-			// Build array starting with class and ending with oldest parent
-			EntriesToProcess.Reset();
-			TSharedPtr<FClassViewerNode> OldestValidParent;
-			for (; CurrentClass; CurrentClass = CurrentClass->GetSuperClass())
-			{
-				TSharedPtr<FClassViewerNode>& MyEntry = VisitedNodes.FindOrAdd(CurrentClass);
-				if (!MyEntry.IsValid())
-				{
-					MyEntry = MakeShared<FClassViewerNode>(CurrentClass);
-					InOutClassPathToNode.FindOrAdd(MyEntry->ClassPath) = MyEntry;
-					EntriesToProcess.Add(TPair<UClass*, TSharedPtr<FClassViewerNode>>(CurrentClass, MyEntry));
-				}
-				else
-				{
-					OldestValidParent = MyEntry;
-					break;
-				}
-			}
-			// Iterate array in reverse starting with oldest parent and ending with class
-			for (int32 i=EntriesToProcess.Num() - 1; i >= 0; --i)
-			{
-				CurrentClass = EntriesToProcess[i].Key;
-				TSharedPtr<FClassViewerNode>& MyEntry = EntriesToProcess[i].Value;
-				if ((i + 1) < EntriesToProcess.Num())
-				{
-					EntriesToProcess[i + 1].Value->AddUniqueChild(MyEntry);
-				}
-				else if (OldestValidParent.IsValid())
-				{
-					OldestValidParent->AddUniqueChild(MyEntry);
-				}
-			}
+			Node = MakeShared<FClassViewerNode>(CurrentClass->GetName(), CurrentClass->GetDisplayNameText().ToString());
 		}
+		SetClassFields(Node, *CurrentClass);
 	}
+	TSharedPtr<FClassViewerNode>* ExistingRoot = InOutClassPathToNode.Find(FName(*UObject::StaticClass()->GetPathName()));
+	check(ExistingRoot && ExistingRoot->IsValid());
+	OutRootNode = *ExistingRoot;
 }
 
 TSharedPtr< FClassViewerNode > FClassHierarchy::FindParent(const TSharedPtr< FClassViewerNode >& InRootNode, FName InParentClassname, const UClass* InParentClass)
@@ -1299,15 +1262,40 @@ void FClassHierarchy::SortChildren( TSharedPtr< FClassViewerNode >& InRootNode)
 	ChildList.Sort(FClassViewerNodeNameLess());
 }
 
-void FClassHierarchy::FindClass(TSharedPtr< FClassViewerNode > InOutClassNode)
+void FClassHierarchy::FindClass(TSharedPtr<FClassViewerNode> InOutClassNode)
 {
 	UClass* Class = FindObject<UClass>(nullptr, *InOutClassNode->ClassPath.ToString());
 
 	if (Class)
 	{
-		InOutClassNode->Blueprint = Cast<UBlueprint>(Class->ClassGeneratedBy);
-		InOutClassNode->Class = Class;
+		SetClassFields(InOutClassNode, *Class);
+	}
+}
+
+void FClassHierarchy::SetClassFields(TSharedPtr<FClassViewerNode>& InOutClassNode, UClass& Class)
+{
+	if (InOutClassNode->Class.Get())
+	{
+		// Already set
+		return;
+	}
+
+	// Fields that can als be set from FAssetData
+	if (InOutClassNode->ClassPath.IsNone())
+	{
+		InOutClassNode->ClassPath = FName(*Class.GetPathName());
+	}
+	if (InOutClassNode->ParentClassPath.IsNone())
+	{
+		if (Class.GetSuperClass())
+		{
+			InOutClassNode->ParentClassPath = FName(*Class.GetSuperClass()->GetPathName());
+		}
 	}
+
+	// UClass-specific fields
+	InOutClassNode->Blueprint = Cast<UBlueprint>(Class.ClassGeneratedBy);
+	InOutClassNode->Class = &Class;
 }
 
 void FClassHierarchy::LoadUnloadedTagData(TSharedPtr<FClassViewerNode>& InOutClassViewerNode, const FAssetData& InAssetData)
@@ -1320,21 +1308,40 @@ void FClassHierarchy::LoadUnloadedTagData(TSharedPtr<FClassViewerNode>& InOutCla
 	}
 	// Create the viewer node. We use the name without _C for both
 	InOutClassViewerNode = MakeShared<FClassViewerNode>(ClassName, ClassDisplayName);
+	SetAssetDataFields(InOutClassViewerNode, InAssetData);
+}
 
-	InOutClassViewerNode->BlueprintAssetPath = InAssetData.ObjectPath;
+void FClassHierarchy::SetAssetDataFields(TSharedPtr<FClassViewerNode>& InOutClassViewerNode, const FAssetData& InAssetData)
+{
+	if (InOutClassViewerNode->UnloadedBlueprintData.IsValid())
+	{
+		// Already set
+		return;
+	}
+
+	// Fields that can also be set from UClass*
 
 	FString ClassObjectPath;
-	if (InAssetData.GetTagValue(FBlueprintTags::GeneratedClassPath, ClassObjectPath))
+	if (InOutClassViewerNode->ClassPath.IsNone())
 	{
-		InOutClassViewerNode->ClassPath = FName(*FPackageName::ExportTextPathToObjectPath(ClassObjectPath));
+		if (InAssetData.GetTagValue(FBlueprintTags::GeneratedClassPath, ClassObjectPath))
+		{
+			InOutClassViewerNode->ClassPath = FName(*FPackageName::ExportTextPathToObjectPath(ClassObjectPath));
+		}
 	}
-
-	FString ParentClassPathString;
-	if (InAssetData.GetTagValue(FBlueprintTags::ParentClassPath, ParentClassPathString))
+	if (InOutClassViewerNode->ParentClassPath.IsNone())
 	{
-		InOutClassViewerNode->ParentClassPath = FName(*FPackageName::ExportTextPathToObjectPath(ParentClassPathString));
+		FString ParentClassPathString;
+		if (InAssetData.GetTagValue(FBlueprintTags::ParentClassPath, ParentClassPathString))
+		{
+			InOutClassViewerNode->ParentClassPath = FName(*FPackageName::ExportTextPathToObjectPath(ParentClassPathString));
+		}
 	}
 
+	// Blueprint-specific fields
+
+	InOutClassViewerNode->BlueprintAssetPath = InAssetData.ObjectPath;
+
 	// It is an unloaded blueprint, so we need to create the structure that will hold the data.
 	TSharedPtr<FUnloadedBlueprintData> UnloadedBlueprintData = MakeShareable( new FUnloadedBlueprintData(InOutClassViewerNode) );
 	InOutClassViewerNode->UnloadedBlueprintData = UnloadedBlueprintData;
@@ -1374,60 +1381,77 @@ void FClassHierarchy::LoadUnloadedTagData(TSharedPtr<FClassViewerNode>& InOutCla
 
 void FClassHierarchy::PopulateClassHierarchy()
 {
-	TArray< TSharedPtr< FClassViewerNode > > RootLevelClasses;
+	// Fetch all classes from AssetRegistry blueprint data (which covers unloaded classes), and in-memory UClasses.
+	// Create a node for each one with unioned data from the AssetRegistry or UClass for that class.
+	// Set parent/child pointers to create a tree, and store this tree in this->ObjectClassRoot
+	TMap<FName, TSharedPtr<FClassViewerNode>> ClassPathToNode;
 
-	FAssetRegistryModule& AssetRegistryModule = FModuleManager::LoadModuleChecked<FAssetRegistryModule>(TEXT("AssetRegistry"));
 
+	// Create a node for every Blueprint class listed in the AssetRegistry and set the Blueprint fields
 	// Retrieve all blueprint classes 
+	FAssetRegistryModule& AssetRegistryModule = FModuleManager::LoadModuleChecked<FAssetRegistryModule>(TEXT("AssetRegistry"));
 	TArray<FAssetData> BlueprintList;
-
 	FARFilter Filter;
 	Filter.ClassNames.Add(UBlueprint::StaticClass()->GetFName());
 	Filter.ClassNames.Add(UAnimBlueprint::StaticClass()->GetFName());
 	Filter.ClassNames.Add(UBlueprintGeneratedClass::StaticClass()->GetFName());
-
 	// Include any Blueprint based objects as well, this includes things like Blutilities, UMG, and GameplayAbility objects
 	Filter.bRecursiveClasses = true;
 	AssetRegistryModule.Get().GetAssets(Filter, BlueprintList);
-
-	TMap<FName, TSharedPtr< FClassViewerNode >> ClassPathToNode;
-	for ( int32 AssetIdx = 0; AssetIdx < BlueprintList.Num(); ++AssetIdx )
+	FString ClassPathString;
+	for (FAssetData& AssetData : BlueprintList)
 	{
-		TSharedPtr< FClassViewerNode > NewNode;
-		LoadUnloadedTagData(NewNode, BlueprintList[AssetIdx]);
-		RootLevelClasses.Add(NewNode);
-
-		check(!NewNode->GetChildrenList().Num());
-		ClassPathToNode.Add(NewNode->ClassPath, NewNode);
-
-		// Find the blueprint if it's loaded.
-		FindClass(NewNode);
+		FName ClassPath;
+		if (AssetData.GetTagValue(FBlueprintTags::GeneratedClassPath, ClassPathString))
+		{
+			ClassPath = FName(*FPackageName::ExportTextPathToObjectPath(ClassPathString));
+		}
+		if (ClassPath.IsNone())
+		{
+			UE_LOG(LogEditorClassViewer, Warning, TEXT("AssetRegistry Blueprint %s is missing tag value for %s. Blueprint will not be available to ClassViewer when unloaded."),
+				*AssetData.ObjectPath.ToString(), *FBlueprintTags::GeneratedClassPath.ToString());
+			continue;
+		}
+		TSharedPtr<FClassViewerNode>& Node = ClassPathToNode.FindOrAdd(ClassPath);
+		if (!Node)
+		{
+			const FString ClassName = AssetData.AssetName.ToString();
+			FString ClassDisplayName = AssetData.GetTagValueRef<FString>(FBlueprintTags::BlueprintDisplayName);
+			if (ClassDisplayName.IsEmpty())
+			{
+				ClassDisplayName = ClassName;
+			}
+			Node = MakeShared<FClassViewerNode>(ClassName, ClassDisplayName);
+		}
+		SetAssetDataFields(Node, AssetData);
 	}
 
-	AddChildren_NoFilter(ObjectClassRoot, ClassPathToNode);
+	// FindOrCreate a node for every loaded UClass, and set the UClass fields
+	CreateNodesForLoadedClasses(ObjectClassRoot, ClassPathToNode);
 
-	RootLevelClasses.Add(ObjectClassRoot);
-
-	// Second pass to link them to parents.
-	for (int32 CurrentNodeIdx = 0; CurrentNodeIdx < RootLevelClasses.Num(); ++CurrentNodeIdx)
+	// Set the parent and child pointers
+	for (TPair<FName, TSharedPtr<FClassViewerNode>>& KVPair : ClassPathToNode)
 	{
-		if(RootLevelClasses[CurrentNodeIdx]->ParentClassPath != NAME_None)
+		TSharedPtr<FClassViewerNode>& Node = KVPair.Value;
+		if (Node == ObjectClassRoot)
 		{
-			// Resolve the parent's class name locally and use it to find the parent's class.
-			FString ParentClassPath = RootLevelClasses[CurrentNodeIdx]->ParentClassPath.ToString();
-			const UClass* ParentClass = FindObject<UClass>(nullptr, *ParentClassPath);
-			TSharedPtr< FClassViewerNode > ParentNode;
-
-			if (TSharedPtr< FClassViewerNode >* ParentNodePtr = ClassPathToNode.Find(RootLevelClasses[CurrentNodeIdx]->ParentClassPath))
-			{
-				// AddUniqueChild makes sure that when a node was generated one by EditorClassHierarchy and one from LoadUnloadedTagData - the proper one is selected
-				ParentNode = *ParentNodePtr;
-				ParentNode->AddUniqueChild(RootLevelClasses[CurrentNodeIdx]);	
-				RootLevelClasses.RemoveAtSwap(CurrentNodeIdx);
-				--CurrentNodeIdx;
-			}
-
+			// No parent expected for the root class
+			continue;
+		}
+		TSharedPtr<FClassViewerNode>* ParentNodePtr = nullptr;
+		if (!Node->ParentClassPath.IsNone())
+		{
+			ParentNodePtr = ClassPathToNode.Find(Node->ParentClassPath);
+		}
+		if (!ParentNodePtr)
+		{
+			UE_LOG(LogEditorClassViewer, Warning, TEXT("Class %s has parent %s, but this parent is not found. The Class will not be shown in ClassViewer."),
+				*KVPair.Key.ToString(), *Node->ParentClassPath.ToString());
+			continue;
 		}
+		TSharedPtr<FClassViewerNode>& ParentNode = *ParentNodePtr;
+		check(ParentNode);
+		ParentNode->AddChild(Node);
 	}
 
 	// Recursively sort the children.
