Prevent REINST classes from being allowed in a class viewer. This bug does not happen consistently, so I've added an ensure to hopefully track down how it is happening.

#rnx

diff --git a/Engine/Source/Editor/ClassViewer/Private/ClassViewerFilter.cpp b/Engine/Source/Editor/ClassViewer/Private/ClassViewerFilter.cpp
--- a/Engine/Source/Editor/ClassViewer/Private/ClassViewerFilter.cpp
+++ b/Engine/Source/Editor/ClassViewer/Private/ClassViewerFilter.cpp
@@ -370,6 +370,13 @@ static bool PassesTextFilter(const FString& InTestString, const TSharedRef<FText
 	return InTextFilter->TestTextFilter(FBasicStringFilterExpressionContext(InTestString));
 }
 
+/** Returns true if the given class is a REINST class (starts with the 'REINST_' prefix) */
+static bool IsReinstClass(const UClass* Class)
+{
+	static const FString ReinstPrefix = TEXT("REINST");
+	return Class && Class->GetFName().ToString().StartsWith(ReinstPrefix);
+}
+
 FClassViewerFilter::FClassViewerFilter(const FClassViewerInitializationOptions& InInitOptions) :
 	TextFilter(MakeShared<FTextFilterExpressionEvaluator>(ETextFilterExpressionEvaluatorMode::BasicString)),
 	FilterFunctions(MakeShared<FClassViewerFilterFuncs>()),
@@ -532,9 +539,14 @@ PRAGMA_ENABLE_DEPRECATION_WARNINGS
 		bPassesAssetReferenceFilter = AssetReferenceFilter->PassesFilter(FAssetData(InClass));
 	}
 
+	// REINST classes cannot be used in any class viewer. 
+	const bool bPassesReinstFilter = !IsReinstClass(InClass);
+	ensureMsgf(bPassesReinstFilter, TEXT("[UE-129601] A 'REINST' class has been given the the Class Viewer! '%s'"), *InClass->GetFName().ToString());
+
 	bool bPassesFilter = bPassesAllowedClasses && bPassesPlaceableFilter && bPassesBlueprintBaseFilter
 		&& bPassesDeveloperFilter && bPassesInternalFilter && bPassesEditorClassFilter 
-		&& bPassesCustomFilter && bPassesGlobalClassFilter && bPassesTextFilter && bPassesAssetReferenceFilter;
+		&& bPassesCustomFilter && bPassesGlobalClassFilter && bPassesTextFilter && bPassesAssetReferenceFilter
+		&& bPassesReinstFilter;
 
 	return bPassesFilter;
 }
