Fixes to properly show cooked blueprints in the class viewer

diff --git a/Engine/Source/Editor/ClassViewer/Private/SClassViewer.cpp b/Engine/Source/Editor/ClassViewer/Private/SClassViewer.cpp
--- a/Engine/Source/Editor/ClassViewer/Private/SClassViewer.cpp
+++ b/Engine/Source/Editor/ClassViewer/Private/SClassViewer.cpp
@@ -161,20 +161,13 @@ private:
 	void OnReloadComplete( EReloadCompleteReason Reason );
 
 	/** 
-	 * Loads the tag data for an unloaded blueprint asset.
+	 * Makes or updates a class viewer node with the data for an unloaded blueprint asset.
 	 *
 	 * @param InOutClassViewerNode		The node to save all the data into.
 	 * @param InAssetData				The asset data to pull the tags from.
+	 * @parma InClassPath				The class path
 	 */
-	void LoadUnloadedTagData(TSharedPtr<FClassViewerNode>& InOutClassViewerNode, const FAssetData& InAssetData);
-
-	/**
-	 * Sets the fields calculated from AssetData on the given Node, if not already set
-	 *
-	 * @param InOutClassViewerNode		The node to save all the data into.
-	 * @param InAssetData				The asset data to pull the tags from.
-	 */
-	void SetAssetDataFields(TSharedPtr<FClassViewerNode>& InOutClassViewerNode, const FAssetData& InAssetData);
+	void CreateOrUpdateUnloadedClassNode(TSharedPtr<FClassViewerNode>& InOutClassViewerNode, const FAssetData& InAssetData, FName InClassPath);
 
 	/**
 	 * Finds the UClass and UBlueprint for the passed in node, utilizing unloaded data to find it.
@@ -908,7 +901,7 @@ public:
 	}
 
 private:
-	FReply OnMouseButtonDoubleClick( const FGeometry& InMyGeometry, const FPointerEvent& InMouseEvent )
+	virtual FReply OnMouseButtonDoubleClick( const FGeometry& InMyGeometry, const FPointerEvent& InMouseEvent ) override
 	{
 		// If in a Class Viewer and it has not been loaded, load the class when double-left clicking.
 		if ( bIsInClassViewer )
@@ -1017,10 +1010,10 @@ static void OnModulesChanged(FName ModuleThatChanged, EModuleChangeReason Reason
 FClassHierarchy::FClassHierarchy()
 {
 	// Register with the Asset Registry to be informed when it is done loading up files.
-	FAssetRegistryModule& AssetRegistryModule = FModuleManager::GetModuleChecked<FAssetRegistryModule>(TEXT("AssetRegistry"));
-	OnFilesLoadedRequestPopulateClassHierarchyDelegateHandle = AssetRegistryModule.Get().OnFilesLoaded().AddStatic( ClassViewer::Helpers::RequestPopulateClassHierarchy );
-	AssetRegistryModule.Get().OnAssetAdded().AddRaw( this, &FClassHierarchy::AddAsset);
-	AssetRegistryModule.Get().OnAssetRemoved().AddRaw( this, &FClassHierarchy::RemoveAsset );
+	IAssetRegistry& AssetRegistry = FModuleManager::LoadModuleChecked<FAssetRegistryModule>(AssetRegistryConstants::ModuleName).Get();
+	OnFilesLoadedRequestPopulateClassHierarchyDelegateHandle = AssetRegistry.OnFilesLoaded().AddStatic( ClassViewer::Helpers::RequestPopulateClassHierarchy );
+	AssetRegistry.OnAssetAdded().AddRaw( this, &FClassHierarchy::AddAsset);
+	AssetRegistry.OnAssetRemoved().AddRaw( this, &FClassHierarchy::RemoveAsset );
 
 	// Register to have Populate called when doing a Reload.
 	FCoreUObjectDelegates::ReloadCompleteDelegate.AddRaw( this, &FClassHierarchy::OnReloadComplete );
@@ -1035,12 +1028,12 @@ FClassHierarchy::FClassHierarchy()
 FClassHierarchy::~FClassHierarchy()
 {
 	// Unregister with the Asset Registry to be informed when it is done loading up files.
-	if( FModuleManager::Get().IsModuleLoaded( TEXT("AssetRegistry") ) )
+	if (FAssetRegistryModule* AssetRegistryModule = FModuleManager::Get().GetModulePtr<FAssetRegistryModule>(AssetRegistryConstants::ModuleName))
 	{
-		FAssetRegistryModule& AssetRegistryModule = FModuleManager::GetModuleChecked<FAssetRegistryModule>(TEXT("AssetRegistry"));
-		AssetRegistryModule.Get().OnFilesLoaded().Remove(OnFilesLoadedRequestPopulateClassHierarchyDelegateHandle);
-		AssetRegistryModule.Get().OnAssetAdded().RemoveAll( this );
-		AssetRegistryModule.Get().OnAssetRemoved().RemoveAll( this );
+		IAssetRegistry& AssetRegistry = AssetRegistryModule->Get();
+		AssetRegistry.OnFilesLoaded().Remove(OnFilesLoadedRequestPopulateClassHierarchyDelegateHandle);
+		AssetRegistry.OnAssetAdded().RemoveAll( this );
+		AssetRegistry.OnAssetRemoved().RemoveAll( this );
 
 		// Unregister to have Populate called when doing a Reload.
 		FCoreUObjectDelegates::ReloadCompleteDelegate.RemoveAll(this);
@@ -1184,21 +1177,12 @@ bool FClassHierarchy::FindAndRemoveNodeByClassPath(const TSharedPtr< FClassViewe
 
 void FClassHierarchy::RemoveAsset(const FAssetData& InRemovedAssetData)
 {
-	FString ClassObjectPath;
-	if (InRemovedAssetData.GetTagValue(FBlueprintTags::GeneratedClassPath, ClassObjectPath))
-	{
-		ClassObjectPath = FPackageName::ExportTextPathToObjectPath(ClassObjectPath);
+	// BPGCs can be missing if it was already deleted prior to the notification being sent. 
+	// Let's try to reconstruct the generated class path from the BP object path.
+	bool bGenerateClassPathIfMissing = true;
+	const FName ClassPath = FEditorClassUtils::GetClassPathFromAsset(InRemovedAssetData, bGenerateClassPathIfMissing);
 
-		if (ClassObjectPath == "None")
-		{
-			// This can happen if the generated class was already deleted prior to 
-			// the notification being sent. Let's try to reconstruct the generated
-			// class name from the object path.
-			ClassObjectPath = InRemovedAssetData.ObjectPath.ToString() + "_C";
-		}
-	}
-
-	if (FindAndRemoveNodeByClassPath(ObjectClassRoot, FName(*ClassObjectPath)))
+	if (!ClassPath.IsNone() && FindAndRemoveNodeByClassPath(ObjectClassRoot, ClassPath))
 	{
 		// All viewers must refresh.
 		ClassViewer::Helpers::RefreshAll();
@@ -1207,44 +1191,36 @@ void FClassHierarchy::RemoveAsset(const FAssetData& InRemovedAssetData)
 
 void FClassHierarchy::AddAsset(const FAssetData& InAddedAssetData)
 {
-	FAssetRegistryModule& AssetRegistryModule = FModuleManager::LoadModuleChecked<FAssetRegistryModule>(TEXT("AssetRegistry"));
-	if ( !AssetRegistryModule.Get().IsLoadingAssets() )
+	IAssetRegistry& AssetRegistry = FModuleManager::LoadModuleChecked<FAssetRegistryModule>(AssetRegistryConstants::ModuleName).Get();
+	if (AssetRegistry.IsLoadingAssets())
 	{
-		TArray<FName> AncestorClassNames;
-		AssetRegistryModule.Get().GetAncestorClassNames(InAddedAssetData.AssetClass, AncestorClassNames);
+		return;
+	}
 
-		if( AncestorClassNames.Contains(UBlueprintCore::StaticClass()->GetFName()) )
-		{
-			FString ClassObjectPath;
-			if (InAddedAssetData.GetTagValue(FBlueprintTags::GeneratedClassPath, ClassObjectPath))
-			{
-				ClassObjectPath = FPackageName::ExportTextPathToObjectPath(ClassObjectPath);
-			}
+	const FName ClassPath = FEditorClassUtils::GetClassPathFromAsset(InAddedAssetData);
 
-			// Make sure that the node does not already exist. There is a bit of double adding going on at times and this prevents it.
-			if(!FindNodeByGeneratedClassPath(ObjectClassRoot, FName(*ClassObjectPath)).IsValid())
-			{
-				TSharedPtr< FClassViewerNode > NewNode;
-				LoadUnloadedTagData(NewNode, InAddedAssetData);
+	// Make sure that the node does not already exist. There is a bit of double adding going on at times and this prevents it.
+	if (!ClassPath.IsNone() && !FindNodeByGeneratedClassPath(ObjectClassRoot, ClassPath).IsValid())
+	{
+		TSharedPtr<FClassViewerNode> NewNode;
+		CreateOrUpdateUnloadedClassNode(NewNode, InAddedAssetData, ClassPath);
 
-				// Find the blueprint if it's loaded.
-				FindClass(NewNode);
+		// Find the blueprint if it's loaded.
+		FindClass(NewNode);
 
-				// Resolve the parent's class name locally and use it to find the parent's class.
-				FString ParentClassPath = NewNode->ParentClassPath.ToString();
-				UClass* ParentClass = FindObject<UClass>(nullptr, *ParentClassPath);
-				TSharedPtr< FClassViewerNode > ParentNode = FindParent(ObjectClassRoot, NewNode->ParentClassPath, ParentClass); 
-				if (ParentNode.IsValid())
-				{
-					ParentNode->AddChild(NewNode);
+		// Resolve the parent's class name locally and use it to find the parent's class.
+		FString ParentClassPath = NewNode->ParentClassPath.ToString();
+		UClass* ParentClass = FindObject<UClass>(nullptr, *ParentClassPath);
+		TSharedPtr< FClassViewerNode > ParentNode = FindParent(ObjectClassRoot, NewNode->ParentClassPath, ParentClass); 
+		if (ParentNode.IsValid())
+		{
+			ParentNode->AddChild(NewNode);
 
-					// Make sure the children are properly sorted.
-					SortChildren(ObjectClassRoot);
+			// Make sure the children are properly sorted.
+			SortChildren(ObjectClassRoot);
 
-					// All Viewers must repopulate.
-					ClassViewer::Helpers::RefreshAll();
-				}
-			}
+			// All Viewers must repopulate.
+			ClassViewer::Helpers::RefreshAll();
 		}
 	}
 }
@@ -1298,21 +1274,19 @@ void FClassHierarchy::SetClassFields(TSharedPtr<FClassViewerNode>& InOutClassNod
 	InOutClassNode->Class = &Class;
 }
 
-void FClassHierarchy::LoadUnloadedTagData(TSharedPtr<FClassViewerNode>& InOutClassViewerNode, const FAssetData& InAssetData)
+void FClassHierarchy::CreateOrUpdateUnloadedClassNode(TSharedPtr<FClassViewerNode>& InOutClassViewerNode, const FAssetData& InAssetData, FName InClassPath)
 {
-	const FString ClassName = InAssetData.AssetName.ToString();
-	FString ClassDisplayName = InAssetData.GetTagValueRef<FString>(FBlueprintTags::BlueprintDisplayName);
-	if (ClassDisplayName.IsEmpty())
+	if (!InOutClassViewerNode)
 	{
-		ClassDisplayName = ClassName;
+		const FString ClassName = InAssetData.AssetName.ToString();
+		FString ClassDisplayName = InAssetData.GetTagValueRef<FString>(FBlueprintTags::BlueprintDisplayName);
+		if (ClassDisplayName.IsEmpty())
+		{
+			ClassDisplayName = ClassName;
+		}
+		InOutClassViewerNode = MakeShared<FClassViewerNode>(ClassName, ClassDisplayName);
 	}
-	// Create the viewer node. We use the name without _C for both
-	InOutClassViewerNode = MakeShared<FClassViewerNode>(ClassName, ClassDisplayName);
-	SetAssetDataFields(InOutClassViewerNode, InAssetData);
-}
 
-void FClassHierarchy::SetAssetDataFields(TSharedPtr<FClassViewerNode>& InOutClassViewerNode, const FAssetData& InAssetData)
-{
 	if (InOutClassViewerNode->UnloadedBlueprintData.IsValid())
 	{
 		// Already set
@@ -1321,23 +1295,8 @@ void FClassHierarchy::SetAssetDataFields(TSharedPtr<FClassViewerNode>& InOutClas
 
 	// Fields that can also be set from UClass*
 
-	if (InOutClassViewerNode->ClassPath.IsNone())
-	{
-		FString GeneratedClassPath;
-		UClass* AssetClass = InAssetData.GetClass();
-		if (AssetClass && AssetClass->IsChildOf(UBlueprintGeneratedClass::StaticClass()))
-		{
-			InOutClassViewerNode->ClassPath = InAssetData.ObjectPath;
-		}
-		else if (InAssetData.GetTagValue(FBlueprintTags::GeneratedClassPath, GeneratedClassPath))
-		{
-			InOutClassViewerNode->ClassPath = FName(*FPackageName::ExportTextPathToObjectPath(GeneratedClassPath));
-		}
-		else
-		{
-			UE_LOG(LogEditorClassViewer, Verbose, TEXT("Failed to set ClassViewerNode ClassPath for %s"), *InAssetData.ObjectPath.ToString());
-		}
-	}
+	InOutClassViewerNode->ClassPath = InClassPath;
+
 	if (InOutClassViewerNode->ParentClassPath.IsNone())
 	{
 		FString ParentClassPathString;
@@ -1352,11 +1311,11 @@ void FClassHierarchy::SetAssetDataFields(TSharedPtr<FClassViewerNode>& InOutClas
 	InOutClassViewerNode->BlueprintAssetPath = InAssetData.ObjectPath;
 
 	// It is an unloaded blueprint, so we need to create the structure that will hold the data.
-	TSharedPtr<FUnloadedBlueprintData> UnloadedBlueprintData = MakeShareable( new FUnloadedBlueprintData(InOutClassViewerNode) );
+	TSharedPtr<FUnloadedBlueprintData> UnloadedBlueprintData = MakeShareable(new FUnloadedBlueprintData(InOutClassViewerNode));
 	InOutClassViewerNode->UnloadedBlueprintData = UnloadedBlueprintData;
 
 	const bool bNormalBlueprintType = InAssetData.GetTagValueRef<FString>(FBlueprintTags::BlueprintType) == TEXT("BPType_Normal");
-	InOutClassViewerNode->UnloadedBlueprintData->SetNormalBlueprintType( bNormalBlueprintType );
+	InOutClassViewerNode->UnloadedBlueprintData->SetNormalBlueprintType(bNormalBlueprintType);
 
 	// Get the class flags.
 	const uint32 ClassFlags = InAssetData.GetTagValueRef<uint32>(FBlueprintTags::ClassFlags);
@@ -1378,44 +1337,38 @@ void FClassHierarchy::PopulateClassHierarchy()
 	// Set parent/child pointers to create a tree, and store this tree in this->ObjectClassRoot
 	TMap<FName, TSharedPtr<FClassViewerNode>> ClassPathToNode;
 
-
 	// Create a node for every Blueprint class listed in the AssetRegistry and set the Blueprint fields
 	// Retrieve all blueprint classes 
-	FAssetRegistryModule& AssetRegistryModule = FModuleManager::LoadModuleChecked<FAssetRegistryModule>(TEXT("AssetRegistry"));
-	TArray<FAssetData> BlueprintList;
-	FARFilter Filter;
-	Filter.ClassNames.Add(UBlueprint::StaticClass()->GetFName());
-	Filter.ClassNames.Add(UAnimBlueprint::StaticClass()->GetFName());
-	Filter.ClassNames.Add(UBlueprintGeneratedClass::StaticClass()->GetFName());
-	// Include any Blueprint based objects as well, this includes things like Blutilities, UMG, and GameplayAbility objects
-	Filter.bRecursiveClasses = true;
-	AssetRegistryModule.Get().GetAssets(Filter, BlueprintList);
-	FString ClassPathString;
-	for (FAssetData& AssetData : BlueprintList)
-	{
-		FName ClassPath;
-		if (AssetData.GetTagValue(FBlueprintTags::GeneratedClassPath, ClassPathString))
-		{
-			ClassPath = FName(*FPackageName::ExportTextPathToObjectPath(ClassPathString));
-		}
-		if (ClassPath.IsNone())
-		{
-			UE_LOG(LogEditorClassViewer, Warning, TEXT("AssetRegistry Blueprint %s is missing tag value for %s. Blueprint will not be available to ClassViewer when unloaded."),
-				*AssetData.ObjectPath.ToString(), *FBlueprintTags::GeneratedClassPath.ToString());
-			continue;
-		}
-		TSharedPtr<FClassViewerNode>& Node = ClassPathToNode.FindOrAdd(ClassPath);
-		if (!Node)
+	{
+		IAssetRegistry& AssetRegistry = FModuleManager::LoadModuleChecked<FAssetRegistryModule>(AssetRegistryConstants::ModuleName).Get();
+		FString ClassPathString;
+
+		TArray<FAssetData> Assets;
+		AssetRegistry.GetAssetsByClass(UBlueprint::StaticClass()->GetFName(), Assets, /*bSearchSubClasses=*/true);
+
+		for (const FAssetData& AssetData : Assets)
 		{
-			const FString ClassName = AssetData.AssetName.ToString();
-			FString ClassDisplayName = AssetData.GetTagValueRef<FString>(FBlueprintTags::BlueprintDisplayName);
-			if (ClassDisplayName.IsEmpty())
+			const FName ClassPath = FEditorClassUtils::GetClassPathFromAssetTag(AssetData);
+			if (!ClassPath.IsNone())
+			{
+				TSharedPtr<FClassViewerNode>& Node = ClassPathToNode.FindOrAdd(ClassPath);
+				CreateOrUpdateUnloadedClassNode(Node, AssetData, ClassPath);
+			}
+			else
 			{
-				ClassDisplayName = ClassName;
+				UE_LOG(LogEditorClassViewer, Warning, TEXT("AssetRegistry Blueprint %s is missing tag value for %s. Blueprint will not be available to ClassViewer when unloaded."),
+					*AssetData.ObjectPath.ToString(), *FBlueprintTags::GeneratedClassPath.ToString());
 			}
-			Node = MakeShared<FClassViewerNode>(ClassName, ClassDisplayName);
 		}
-		SetAssetDataFields(Node, AssetData);
+
+		Assets.Reset();
+		AssetRegistry.GetAssetsByClass(UBlueprintGeneratedClass::StaticClass()->GetFName(), Assets, /*bSearchSubClasses=*/true);
+
+		for (const FAssetData& AssetData : Assets)
+		{
+			TSharedPtr<FClassViewerNode>& Node = ClassPathToNode.FindOrAdd(AssetData.ObjectPath);
+			CreateOrUpdateUnloadedClassNode(Node, AssetData, AssetData.ObjectPath);
+		}
 	}
 
 	// FindOrCreate a node for every loaded UClass, and set the UClass fields
@@ -1813,6 +1766,9 @@ void SClassViewer::OnClassViewerSelectionChanged( TSharedPtr<FClassViewerNode> I
 		if ( bEnableClassDynamicLoading && !Class && Item->UnloadedBlueprintData.IsValid() )
 		{
 			ClassViewer::Helpers::LoadClass( Item );
+
+			// Populate the tree/list so any changes to previously unloaded classes will be reflected.
+			Refresh();
 		}
 
 		// Check if the item passes the filter
@@ -2141,10 +2097,10 @@ FReply SClassViewer::OnDragDetected( const FGeometry& Geometry, const FPointerEv
 			}
 			else if (Item->BlueprintAssetPath != NAME_None)
 			{
-				FAssetRegistryModule& AssetRegistryModule = FModuleManager::LoadModuleChecked<FAssetRegistryModule>(TEXT("AssetRegistry"));
+				IAssetRegistry& AssetRegistry = FModuleManager::LoadModuleChecked<FAssetRegistryModule>(AssetRegistryConstants::ModuleName).Get();
 
 				// Pull asset data out of asset registry
-				const FAssetData AssetData = AssetRegistryModule.Get().GetAssetByObjectPath(Item->BlueprintAssetPath);
+				const FAssetData AssetData = AssetRegistry.GetAssetByObjectPath(Item->BlueprintAssetPath);
 				return FReply::Handled().BeginDragDrop(FContentBrowserDataDragDropOp::Legacy_New(MakeArrayView(&AssetData, 1)));
 			}
 		}
diff --git a/Engine/Source/Editor/UnrealEd/Private/EditorClassUtils.cpp b/Engine/Source/Editor/UnrealEd/Private/EditorClassUtils.cpp
--- a/Engine/Source/Editor/UnrealEd/Private/EditorClassUtils.cpp
+++ b/Engine/Source/Editor/UnrealEd/Private/EditorClassUtils.cpp
@@ -6,6 +6,7 @@
 #include "EditorStyleSet.h"
 #include "Engine/Blueprint.h"
 #include "Editor.h"
+#include "AssetRegistryModule.h"
 
 #include "IDocumentationPage.h"
 #include "IDocumentation.h"
@@ -163,6 +164,64 @@ UClass* FEditorClassUtils::GetClassFromString(const FString& ClassName)
 	return Class;
 }
 
+bool FEditorClassUtils::IsBlueprintAsset(const FAssetData& InAssetData, bool* bOutIsBPGC /*= nullptr*/)
+{
+	bool bIsBP = (InAssetData.AssetClass == UBlueprint::StaticClass()->GetFName());
+	bool bIsBPGC = (InAssetData.AssetClass == UBlueprintGeneratedClass::StaticClass()->GetFName());
+
+	if (!bIsBP && !bIsBPGC)
+	{
+		IAssetRegistry& AssetRegistry = FModuleManager::LoadModuleChecked<FAssetRegistryModule>(AssetRegistryConstants::ModuleName).Get();
+		TArray<FName> AncestorClassNames;
+		AssetRegistry.GetAncestorClassNames(InAssetData.AssetClass, AncestorClassNames);
+
+		if (AncestorClassNames.Contains(UBlueprint::StaticClass()->GetFName()))
+		{
+			bIsBP = true;
+		}
+		else if (AncestorClassNames.Contains(UBlueprintGeneratedClass::StaticClass()->GetFName()))
+		{
+			bIsBPGC = true;
+		}
+	}
+
+	if (bOutIsBPGC)
+	{
+		*bOutIsBPGC = bIsBPGC;
+	}
+
+	return bIsBP || bIsBPGC;
+}
+
+FName FEditorClassUtils::GetClassPathFromAssetTag(const FAssetData& InAssetData)
+{
+	const FString GeneratedClassPath = InAssetData.GetTagValueRef<FString>(FBlueprintTags::GeneratedClassPath);
+	return FName(FPackageName::ExportTextPathToObjectPath(FStringView(GeneratedClassPath)));
+}
+
+FName FEditorClassUtils::GetClassPathFromAsset(const FAssetData& InAssetData, bool bGenerateClassPathIfMissing /*= false*/)
+{
+	bool bIsBPGC = false;
+	const bool bIsBP = IsBlueprintAsset(InAssetData, &bIsBPGC);
+
+	if (bIsBPGC)
+	{
+		return InAssetData.ObjectPath;
+	}
+	else if (bIsBP)
+	{
+		FName ClassPath = GetClassPathFromAssetTag(InAssetData);
+		if (bGenerateClassPathIfMissing && ClassPath.IsNone())
+		{
+			FNameBuilder ClassPathBuilder(InAssetData.ObjectPath);
+			ClassPathBuilder << "_C";
+			ClassPath = FName(ClassPathBuilder.ToString());
+		}
+		return ClassPath;
+	}
+	return NAME_None;
+}
+
 void FEditorClassUtils::GetImplementedInterfaceClassPathsFromAsset(const struct FAssetData& InAssetData, TArray<FString>& OutClassPaths)
 {
 	if (!InAssetData.IsValid())
diff --git a/Engine/Source/Editor/UnrealEd/Public/EditorClassUtils.h b/Engine/Source/Editor/UnrealEd/Public/EditorClassUtils.h
--- a/Engine/Source/Editor/UnrealEd/Public/EditorClassUtils.h
+++ b/Engine/Source/Editor/UnrealEd/Public/EditorClassUtils.h
@@ -7,6 +7,7 @@
 #include "Widgets/SToolTip.h"
 
 template< typename ObjectType > class TAttribute;
+struct FAssetData;
 
 namespace FEditorClassUtils
 {
@@ -100,11 +101,38 @@ namespace FEditorClassUtils
 	 */
 	UNREALED_API UClass* GetClassFromString(const FString& ClassName);
 
+	/**
+	 * Returns whether the specified asset is a UBlueprint or UBlueprintGeneratedClass (or any of their derived classes)
+	 * 
+	 * @param	InAssetData		Reference to an asset data entry
+	 * @param	bOutIsBPGC		Outputs whether the asset is a BlueprintGeneratedClass (or any of its derived classes)
+	 * @return					Whether the specified asset is a UBlueprint or UBlueprintGeneratedClass (or any of their derived classes)
+	 */
+	UNREALED_API bool IsBlueprintAsset(const FAssetData& InAssetData, bool* bOutIsBPGC = nullptr);
+
+	/**
+	 * Gets the class path from the asset tag (i.e. GeneratedClassPath tag on blueprints)
+	 * 
+	 * @param	InAssetData		Reference to an asset data entry.
+	 * @return					Class path or None if the asset cannot or doesn't have a class associated with it
+	 */
+	UNREALED_API FName GetClassPathFromAssetTag(const FAssetData& InAssetData);
+
+	/**
+	 * Gets the object path of the class associated with the specified asset 
+	 * (i.e. the BlueprintGeneratedClass of a Blueprint asset or the BlueprintGeneratedClass asset itself)
+	 * 
+	 * @param	InAssetData					Reference to an asset data entry
+	 * @param	bGenerateClassPathIfMissing	Whether to generate a class path if the class is missing (and the asset can have a class associated with it)
+	 * @return								Class path or None if the asset cannot or doesn't have a class associated with it
+	 */
+	UNREALED_API FName GetClassPathFromAsset(const FAssetData& InAssetData, bool bGenerateClassPathIfMissing = false);
+
 	/**
 	 * Fetches the set of interface class object paths from an asset data entry containing the appropriate asset tag(s).
 	 * 
 	 * @param	InAssetData		Reference to an asset data entry.
 	 * @param	OutClassPaths	One or more interface class object paths, or empty if the corresponding asset tag(s) were not found.
 	 */
-	UNREALED_API void GetImplementedInterfaceClassPathsFromAsset(const struct FAssetData& InAssetData, TArray<FString>& OutClassPaths);
+	UNREALED_API void GetImplementedInterfaceClassPathsFromAsset(const FAssetData& InAssetData, TArray<FString>& OutClassPaths);
 };
