Remove an ensure that was in place to track down where REINST classes were being placed into the class viewer. The reason has been found and we know how to fix it so the ensure can be removed.

diff --git a/Engine/Source/Editor/ClassViewer/Private/ClassViewerFilter.cpp b/Engine/Source/Editor/ClassViewer/Private/ClassViewerFilter.cpp
--- a/Engine/Source/Editor/ClassViewer/Private/ClassViewerFilter.cpp
+++ b/Engine/Source/Editor/ClassViewer/Private/ClassViewerFilter.cpp
@@ -541,7 +541,6 @@ PRAGMA_ENABLE_DEPRECATION_WARNINGS
 
 	// REINST classes cannot be used in any class viewer. 
 	const bool bPassesReinstFilter = !IsReinstClass(InClass);
-	ensureMsgf(bPassesReinstFilter, TEXT("[UE-129601] A 'REINST' class has been given the the Class Viewer! '%s'"), *InClass->GetFName().ToString());
 
 	bool bPassesFilter = bPassesAllowedClasses && bPassesPlaceableFilter && bPassesBlueprintBaseFilter
 		&& bPassesDeveloperFilter && bPassesInternalFilter && bPassesEditorClassFilter 
