Repopulating ClassViewer cache when scripts have finished compiling in the editor

diff --git a/Engine/Source/Editor/ClassViewer/Private/SClassViewer.cpp b/Engine/Source/Editor/ClassViewer/Private/SClassViewer.cpp
--- a/Engine/Source/Editor/ClassViewer/Private/SClassViewer.cpp
+++ b/Engine/Source/Editor/ClassViewer/Private/SClassViewer.cpp
@@ -2569,6 +2569,11 @@ bool SClassViewer::SupportsKeyboardFocus() const
 	return true;
 }
 
+void SClassViewer::RequestPopulateClassHierarchy()
+{
+	ClassViewer::Helpers::RequestPopulateClassHierarchy();
+}
+
 void SClassViewer::DestroyClassHierarchy()
 {
 	ClassViewer::Helpers::DestroyClassHierachy();
diff --git a/Engine/Source/Editor/ClassViewer/Public/SClassViewer.h b/Engine/Source/Editor/ClassViewer/Public/SClassViewer.h
--- a/Engine/Source/Editor/ClassViewer/Public/SClassViewer.h
+++ b/Engine/Source/Editor/ClassViewer/Public/SClassViewer.h
@@ -66,6 +66,9 @@ public:
 	/** Sends a requests to the Class Viewer to refresh itself the next chance it gets */
 	CLASSVIEWER_API void Refresh();
 
+	/** Requests a repopulation of the internal Class Hierarchy database */
+	CLASSVIEWER_API static void RequestPopulateClassHierarchy();
+
 	/** Destroys the internal Class Hierarchy database */
 	static void DestroyClassHierarchy();
 private:
