ClassViewerFilter: Make the class Display Name searchable

diff --git a/Engine/Source/Editor/ClassViewer/Private/ClassViewerFilter.cpp b/Engine/Source/Editor/ClassViewer/Private/ClassViewerFilter.cpp
--- a/Engine/Source/Editor/ClassViewer/Private/ClassViewerFilter.cpp
+++ b/Engine/Source/Editor/ClassViewer/Private/ClassViewerFilter.cpp
@@ -523,7 +523,7 @@ PRAGMA_ENABLE_DEPRECATION_WARNINGS
 	if (bCheckTextFilter && (TextFilter->GetFilterType() != ETextFilterExpressionType::Empty))
 	{
 		FString ClassNameWithCppPrefix = FString::Printf(TEXT("%s%s"), InClass->GetPrefixCPP(), *InClass->GetName());
-		bPassesTextFilter = PassesTextFilter(InClass->GetName(), TextFilter) || PassesTextFilter(ClassNameWithCppPrefix, TextFilter);
+		bPassesTextFilter = PassesTextFilter(InClass->GetName(), TextFilter) || PassesTextFilter(ClassNameWithCppPrefix, TextFilter) || PassesTextFilter(InClass->GetDisplayNameText().ToString(), TextFilter);
 
 		// If the class is deprecated, try searching without the deprecated name inserted, in case a user typed a string
 		if (!bPassesTextFilter && InClass->HasAnyClassFlags(CLASS_Deprecated))
