Removed unnecessary inline forward declarations from ClassViewerModule.h.

#rnx

diff --git a/Engine/Source/Editor/ClassViewer/Public/ClassViewerModule.h b/Engine/Source/Editor/ClassViewer/Public/ClassViewerModule.h
--- a/Engine/Source/Editor/ClassViewer/Public/ClassViewerModule.h
+++ b/Engine/Source/Editor/ClassViewer/Public/ClassViewerModule.h
@@ -6,9 +6,10 @@
 #include "AssetRegistry/AssetData.h"
 #include "Modules/ModuleInterface.h"
 
+class FClassViewerFilterFuncs;
 class IClassViewerFilter;
 class IPropertyHandle;
-class FClassViewerFilterFuncs;
+class SWidget;
 
 /** Delegate used with the Class Viewer in 'class picking' mode.  You'll bind a delegate when the
     class viewer widget is created, which will be fired off when a class is selected in the list */
@@ -64,10 +65,10 @@
 public:
 	/** [Deprecated] The filter to use on classes in this instance. */
 	UE_DEPRECATED(5.0, "Please add to the ClassFilters array member instead.")
-	TSharedPtr<class IClassViewerFilter> ClassFilter;
+	TSharedPtr<IClassViewerFilter> ClassFilter;
 
 	/** The filter(s) to use on classes in this instance. */
-	TArray<TSharedRef<class IClassViewerFilter>> ClassFilters;
+	TArray<TSharedRef<IClassViewerFilter>> ClassFilters;
 
 	/** Mode to operate in */
 	EClassViewerMode::Type Mode;
@@ -109,7 +110,7 @@
 	FText ViewerTitleString;
 
 	/** The property this class viewer be working on. */
-	TSharedPtr<class IPropertyHandle> PropertyHandle;
+	TSharedPtr<IPropertyHandle> PropertyHandle;
 
 	/** The passed in property handle will be used to gather referencing assets. If additional referencing assets should be reported, supply them here. */
 	TArray<FAssetData> AdditionalReferencingAssets;
@@ -183,8 +184,7 @@
 	 *
 	 * @return	New class viewer widget
 	 */
-	virtual TSharedRef<class SWidget> CreateClassViewer(const FClassViewerInitializationOptions& InitOptions,
-		const FOnClassPicked& OnClassPickedDelegate );
+	virtual TSharedRef<SWidget> CreateClassViewer(const FClassViewerInitializationOptions& InitOptions, const FOnClassPicked& OnClassPickedDelegate );
 
 	/** 
 	 * Create a new class filter from the given initialization options.
