Deprecating ANY_PACKAGE.

diff --git a/Engine/Source/Editor/ClassViewer/Private/SClassViewer.cpp b/Engine/Source/Editor/ClassViewer/Private/SClassViewer.cpp
--- a/Engine/Source/Editor/ClassViewer/Private/SClassViewer.cpp
+++ b/Engine/Source/Editor/ClassViewer/Private/SClassViewer.cpp
@@ -135,27 +135,27 @@ public:
 	 *
 	 *	@return The parent node.
 	 */
-	TSharedPtr< FClassViewerNode > FindParent(const TSharedPtr< FClassViewerNode >& InRootNode, FName InParentClassname, const UClass* InParentClass);
+	TSharedPtr< FClassViewerNode > FindParent(const TSharedPtr< FClassViewerNode >& InRootNode, FTopLevelAssetPath InParentClassname, const UClass* InParentClass);
 
 	/** Updates the Class of a node. Uses the generated class package name to find the node.
 	 *	@param InGeneratedClassPath		The path of the generated class to find the node for.
 	 *	@param InNewClass				The class to update the node with.
 	*/
-	void UpdateClassInNode(FName InGeneratedClassPath, UClass* InNewClass, UBlueprint* InNewBluePrint );
+	void UpdateClassInNode(FTopLevelAssetPath InGeneratedClassPath, UClass* InNewClass, UBlueprint* InNewBluePrint );
 
 	/** Finds the node, recursively going deeper into the hierarchy. Does so by comparing generated class package names.
 	 *	@param InGeneratedClassPath		The path of the generated class to find the node for.
 	 *
 	 *	@return The node.
 	 */
-	TSharedPtr< FClassViewerNode > FindNodeByGeneratedClassPath(const TSharedPtr< FClassViewerNode >& InRootNode, FName InGeneratedClassPath);
+	TSharedPtr< FClassViewerNode > FindNodeByGeneratedClassPath(const TSharedPtr< FClassViewerNode >& InRootNode, FTopLevelAssetPath InGeneratedClassPath);
 
 private:
 	/** Adds UClass information to The node in InOutClassPathToNode for every loaded class, creating the node if it does not exist. Does not filter classes.
 	 *	@param OutRootNode						Out parameter that receives the node pointer for the root UClass.
 	 *  @param InOutClassPathToNode				Map containing all nodes.
 	 */	
-	void CreateNodesForLoadedClasses( TSharedPtr<FClassViewerNode>& OutRootNode, TMap<FName, TSharedPtr< FClassViewerNode >>& InOutClassPathToNode );
+	void CreateNodesForLoadedClasses( TSharedPtr<FClassViewerNode>& OutRootNode, TMap<FTopLevelAssetPath, TSharedPtr< FClassViewerNode >>& InOutClassPathToNode );
 
 	/** Called when reload has finished */
 	void OnReloadComplete( EReloadCompleteReason Reason );
@@ -167,7 +167,7 @@ private:
 	 * @param InAssetData				The asset data to pull the tags from.
 	 * @parma InClassPath				The class path
 	 */
-	void CreateOrUpdateUnloadedClassNode(TSharedPtr<FClassViewerNode>& InOutClassViewerNode, const FAssetData& InAssetData, FName InClassPath);
+	void CreateOrUpdateUnloadedClassNode(TSharedPtr<FClassViewerNode>& InOutClassViewerNode, const FAssetData& InAssetData, FTopLevelAssetPath InClassPath);
 
 	/**
 	 * Finds the UClass and UBlueprint for the passed in node, utilizing unloaded data to find it.
@@ -187,7 +187,7 @@ private:
 	 *
 	 * @return Returns true if the asset was found and deleted successfully.
 	 */
-	bool FindAndRemoveNodeByClassPath(const TSharedPtr< FClassViewerNode >& InRootNode, FName InClassPath);
+	bool FindAndRemoveNodeByClassPath(const TSharedPtr< FClassViewerNode >& InRootNode, FTopLevelAssetPath InClassPath);
 
 	/** Callback registered to the Asset Registry to be notified when an asset is added. */
 	void AddAsset(const FAssetData& InAddedAssetData);
@@ -221,7 +221,7 @@ namespace ClassViewer
 
 		// Pre-declare these functions.
 		static UBlueprint* GetBlueprint( UClass* InClass );
-		static void UpdateClassInNode(FName InGeneratedClassPath, UClass* InNewClass, UBlueprint* InNewBluePrint );
+		static void UpdateClassInNode(FTopLevelAssetPath InGeneratedClassPath, UClass* InNewClass, UBlueprint* InNewBluePrint );
 
 		/** Util class to checks if a particular class can be made into a Blueprint, ignores deprecation
 		 *
@@ -627,7 +627,7 @@ namespace ClassViewer
 			{
 				FMessageLog EditorErrors("EditorErrors");
 				FFormatNamedArguments Arguments;
-				Arguments.Add(TEXT("ObjectName"), FText::FromName(InOutClassNode->ClassPath));
+				Arguments.Add(TEXT("ObjectName"), FText::FromString(InOutClassNode->ClassPath.ToString()));
 				EditorErrors.Error(FText::Format(LOCTEXT("PackageLoadFail", "Failed to load class {ObjectName}"), Arguments));
 			}
 		}
@@ -683,7 +683,7 @@ namespace ClassViewer
 		*	@param InGeneratedClassPath			The name of the generated class to find the node for.
 		*	@param InNewClass					The class to update the node with.
 		*/
-		static void UpdateClassInNode(FName InGeneratedClassPath, UClass* InNewClass, UBlueprint* InNewBluePrint )
+		static void UpdateClassInNode(FTopLevelAssetPath InGeneratedClassPath, UClass* InNewClass, UBlueprint* InNewBluePrint )
 		{
 			ClassHierarchy->UpdateClassInNode(InGeneratedClassPath, InNewClass, InNewBluePrint );
 		}
@@ -829,9 +829,9 @@ public:
 					UPackage*  Package  = Class->GetOutermost();
 					ToolTip = FEditorClassUtils::GetTooltip(Class);
 				}
-				else if (AssociatedNode->ClassPath != NAME_None)
+				else if (!AssociatedNode->ClassPath.IsNull())
 				{
-					ToolTip = SNew(SToolTip).Text(FText::FromName(AssociatedNode->ClassPath));
+					ToolTip = SNew(SToolTip).Text(FText::FromString(AssociatedNode->ClassPath.ToString()));
 				}
 
 				return ToolTip;
@@ -1054,7 +1054,7 @@ void FClassHierarchy::OnReloadComplete(EReloadCompleteReason Reason)
 	ClassViewer::Helpers::RequestPopulateClassHierarchy();
 }
 
-void FClassHierarchy::CreateNodesForLoadedClasses(TSharedPtr<FClassViewerNode>& OutRootNode, TMap<FName, TSharedPtr< FClassViewerNode >>& InOutClassPathToNode)
+void FClassHierarchy::CreateNodesForLoadedClasses(TSharedPtr<FClassViewerNode>& OutRootNode, TMap<FTopLevelAssetPath, TSharedPtr< FClassViewerNode >>& InOutClassPathToNode)
 {
 	for (TObjectIterator<UClass> ClassIt; ClassIt; ++ClassIt)
 	{
@@ -1065,19 +1065,19 @@ void FClassHierarchy::CreateNodesForLoadedClasses(TSharedPtr<FClassViewerNode>&
 		{
 			continue;
 		}
-		TSharedPtr<FClassViewerNode>& Node = InOutClassPathToNode.FindOrAdd(FName(*CurrentClass->GetPathName()));
+		TSharedPtr<FClassViewerNode>& Node = InOutClassPathToNode.FindOrAdd(CurrentClass->GetClassPathName());
 		if (!Node)
 		{
 			Node = MakeShared<FClassViewerNode>(CurrentClass->GetName(), CurrentClass->GetDisplayNameText().ToString());
 		}
 		SetClassFields(Node, *CurrentClass);
 	}
-	TSharedPtr<FClassViewerNode>* ExistingRoot = InOutClassPathToNode.Find(FName(*UObject::StaticClass()->GetPathName()));
+	TSharedPtr<FClassViewerNode>* ExistingRoot = InOutClassPathToNode.Find(UObject::StaticClass()->GetClassPathName());
 	check(ExistingRoot && ExistingRoot->IsValid());
 	OutRootNode = *ExistingRoot;
 }
 
-TSharedPtr< FClassViewerNode > FClassHierarchy::FindParent(const TSharedPtr< FClassViewerNode >& InRootNode, FName InParentClassname, const UClass* InParentClass)
+TSharedPtr< FClassViewerNode > FClassHierarchy::FindParent(const TSharedPtr< FClassViewerNode >& InRootNode, FTopLevelAssetPath InParentClassname, const UClass* InParentClass)
 {
 	// Check if the current node is the parent classname that is being searched for.
 	if(InRootNode->ClassPath == InParentClassname)
@@ -1117,7 +1117,7 @@ TSharedPtr< FClassViewerNode > FClassHierarchy::FindParent(const TSharedPtr< FCl
 	return ReturnNode;
 }
 
-TSharedPtr< FClassViewerNode > FClassHierarchy::FindNodeByGeneratedClassPath(const TSharedPtr< FClassViewerNode >& InRootNode, FName InGeneratedClassPath)
+TSharedPtr< FClassViewerNode > FClassHierarchy::FindNodeByGeneratedClassPath(const TSharedPtr< FClassViewerNode >& InRootNode, FTopLevelAssetPath InGeneratedClassPath)
 {
 	if(InRootNode->ClassPath == InGeneratedClassPath)
 	{
@@ -1141,7 +1141,7 @@ TSharedPtr< FClassViewerNode > FClassHierarchy::FindNodeByGeneratedClassPath(con
 	return ReturnNode;
 }
 
-void FClassHierarchy::UpdateClassInNode(FName InGeneratedClassPath, UClass* InNewClass, UBlueprint* InNewBluePrint )
+void FClassHierarchy::UpdateClassInNode(FTopLevelAssetPath InGeneratedClassPath, UClass* InNewClass, UBlueprint* InNewBluePrint )
 {
 	TSharedPtr< FClassViewerNode > Node = FindNodeByGeneratedClassPath(ObjectClassRoot, InGeneratedClassPath);
 
@@ -1152,7 +1152,7 @@ void FClassHierarchy::UpdateClassInNode(FName InGeneratedClassPath, UClass* InNe
 	}
 }
 
-bool FClassHierarchy::FindAndRemoveNodeByClassPath(const TSharedPtr< FClassViewerNode >& InRootNode, FName InClassPath)
+bool FClassHierarchy::FindAndRemoveNodeByClassPath(const TSharedPtr< FClassViewerNode >& InRootNode, FTopLevelAssetPath InClassPath)
 {
 	bool bReturnValue = false;
 
@@ -1180,9 +1180,9 @@ void FClassHierarchy::RemoveAsset(const FAssetData& InRemovedAssetData)
 	// BPGCs can be missing if it was already deleted prior to the notification being sent. 
 	// Let's try to reconstruct the generated class path from the BP object path.
 	bool bGenerateClassPathIfMissing = true;
-	const FName ClassPath = FEditorClassUtils::GetClassPathFromAsset(InRemovedAssetData, bGenerateClassPathIfMissing);
+	const FTopLevelAssetPath ClassPath = FEditorClassUtils::GetClassPathNameFromAsset(InRemovedAssetData, bGenerateClassPathIfMissing);
 
-	if (!ClassPath.IsNone() && FindAndRemoveNodeByClassPath(ObjectClassRoot, ClassPath))
+	if (!ClassPath.IsNull() && FindAndRemoveNodeByClassPath(ObjectClassRoot, ClassPath))
 	{
 		// All viewers must refresh.
 		ClassViewer::Helpers::RefreshAll();
@@ -1197,10 +1197,10 @@ void FClassHierarchy::AddAsset(const FAssetData& InAddedAssetData)
 		return;
 	}
 
-	const FName ClassPath = FEditorClassUtils::GetClassPathFromAsset(InAddedAssetData);
+	const FTopLevelAssetPath ClassPath = FEditorClassUtils::GetClassPathNameFromAsset(InAddedAssetData);
 
 	// Make sure that the node does not already exist. There is a bit of double adding going on at times and this prevents it.
-	if (!ClassPath.IsNone() && !FindNodeByGeneratedClassPath(ObjectClassRoot, ClassPath).IsValid())
+	if (!ClassPath.IsNull() && !FindNodeByGeneratedClassPath(ObjectClassRoot, ClassPath).IsValid())
 	{
 		TSharedPtr<FClassViewerNode> NewNode;
 		CreateOrUpdateUnloadedClassNode(NewNode, InAddedAssetData, ClassPath);
@@ -1257,15 +1257,15 @@ void FClassHierarchy::SetClassFields(TSharedPtr<FClassViewerNode>& InOutClassNod
 	}
 
 	// Fields that can als be set from FAssetData
-	if (InOutClassNode->ClassPath.IsNone())
+	if (InOutClassNode->ClassPath.IsNull())
 	{
-		InOutClassNode->ClassPath = FName(*Class.GetPathName());
+		InOutClassNode->ClassPath = Class.GetClassPathName();
 	}
-	if (InOutClassNode->ParentClassPath.IsNone())
+	if (InOutClassNode->ParentClassPath.IsNull())
 	{
 		if (Class.GetSuperClass())
 		{
-			InOutClassNode->ParentClassPath = FName(*Class.GetSuperClass()->GetPathName());
+			InOutClassNode->ParentClassPath = Class.GetSuperClass()->GetClassPathName();
 		}
 	}
 
@@ -1274,7 +1274,7 @@ void FClassHierarchy::SetClassFields(TSharedPtr<FClassViewerNode>& InOutClassNod
 	InOutClassNode->Class = &Class;
 }
 
-void FClassHierarchy::CreateOrUpdateUnloadedClassNode(TSharedPtr<FClassViewerNode>& InOutClassViewerNode, const FAssetData& InAssetData, FName InClassPath)
+void FClassHierarchy::CreateOrUpdateUnloadedClassNode(TSharedPtr<FClassViewerNode>& InOutClassViewerNode, const FAssetData& InAssetData, FTopLevelAssetPath InClassPath)
 {
 	if (!InOutClassViewerNode)
 	{
@@ -1297,12 +1297,12 @@ void FClassHierarchy::CreateOrUpdateUnloadedClassNode(TSharedPtr<FClassViewerNod
 
 	InOutClassViewerNode->ClassPath = InClassPath;
 
-	if (InOutClassViewerNode->ParentClassPath.IsNone())
+	if (InOutClassViewerNode->ParentClassPath.IsNull())
 	{
 		FString ParentClassPathString;
 		if (InAssetData.GetTagValue(FBlueprintTags::ParentClassPath, ParentClassPathString))
 		{
-			InOutClassViewerNode->ParentClassPath = FName(*FPackageName::ExportTextPathToObjectPath(ParentClassPathString));
+			InOutClassViewerNode->ParentClassPath = FTopLevelAssetPath(*FPackageName::ExportTextPathToObjectPath(ParentClassPathString));
 		}
 	}
 
@@ -1335,7 +1335,7 @@ void FClassHierarchy::PopulateClassHierarchy()
 	// Fetch all classes from AssetRegistry blueprint data (which covers unloaded classes), and in-memory UClasses.
 	// Create a node for each one with unioned data from the AssetRegistry or UClass for that class.
 	// Set parent/child pointers to create a tree, and store this tree in this->ObjectClassRoot
-	TMap<FName, TSharedPtr<FClassViewerNode>> ClassPathToNode;
+	TMap<FTopLevelAssetPath, TSharedPtr<FClassViewerNode>> ClassPathToNode;
 
 	// Create a node for every Blueprint class listed in the AssetRegistry and set the Blueprint fields
 	// Retrieve all blueprint classes 
@@ -1344,12 +1344,12 @@ void FClassHierarchy::PopulateClassHierarchy()
 		FString ClassPathString;
 
 		TArray<FAssetData> Assets;
-		AssetRegistry.GetAssetsByClass(UBlueprint::StaticClass()->GetFName(), Assets, /*bSearchSubClasses=*/true);
+		AssetRegistry.GetAssetsByClass(UBlueprint::StaticClass()->GetClassPathName(), Assets, /*bSearchSubClasses=*/true);
 
 		for (const FAssetData& AssetData : Assets)
 		{
-			const FName ClassPath = FEditorClassUtils::GetClassPathFromAssetTag(AssetData);
-			if (!ClassPath.IsNone())
+			const FTopLevelAssetPath ClassPath = FEditorClassUtils::GetClassPathNameFromAssetTag(AssetData);
+			if (!ClassPath.IsNull())
 			{
 				TSharedPtr<FClassViewerNode>& Node = ClassPathToNode.FindOrAdd(ClassPath);
 				CreateOrUpdateUnloadedClassNode(Node, AssetData, ClassPath);
@@ -1362,12 +1362,13 @@ void FClassHierarchy::PopulateClassHierarchy()
 		}
 
 		Assets.Reset();
-		AssetRegistry.GetAssetsByClass(UBlueprintGeneratedClass::StaticClass()->GetFName(), Assets, /*bSearchSubClasses=*/true);
+		AssetRegistry.GetAssetsByClass(UBlueprintGeneratedClass::StaticClass()->GetClassPathName(), Assets, /*bSearchSubClasses=*/true);
 
 		for (const FAssetData& AssetData : Assets)
 		{
-			TSharedPtr<FClassViewerNode>& Node = ClassPathToNode.FindOrAdd(AssetData.ObjectPath);
-			CreateOrUpdateUnloadedClassNode(Node, AssetData, AssetData.ObjectPath);
+			FTopLevelAssetPath ClassPathNameFromAssetPath(AssetData.ObjectPath.ToString());
+			TSharedPtr<FClassViewerNode>& Node = ClassPathToNode.FindOrAdd(ClassPathNameFromAssetPath);
+			CreateOrUpdateUnloadedClassNode(Node, AssetData, ClassPathNameFromAssetPath);
 		}
 	}
 
@@ -1375,7 +1376,7 @@ void FClassHierarchy::PopulateClassHierarchy()
 	CreateNodesForLoadedClasses(ObjectClassRoot, ClassPathToNode);
 
 	// Set the parent and child pointers
-	for (TPair<FName, TSharedPtr<FClassViewerNode>>& KVPair : ClassPathToNode)
+	for (TPair<FTopLevelAssetPath, TSharedPtr<FClassViewerNode>>& KVPair : ClassPathToNode)
 	{
 		TSharedPtr<FClassViewerNode>& Node = KVPair.Value;
 		if (Node == ObjectClassRoot)
@@ -1384,7 +1385,7 @@ void FClassHierarchy::PopulateClassHierarchy()
 			continue;
 		}
 		TSharedPtr<FClassViewerNode>* ParentNodePtr = nullptr;
-		if (!Node->ParentClassPath.IsNone())
+		if (!Node->ParentClassPath.IsNull())
 		{
 			ParentNodePtr = ClassPathToNode.Find(Node->ParentClassPath);
 		}
@@ -2328,7 +2329,7 @@ int32 SClassViewer::CountTreeItems(FClassViewerNode* Node)
 
 void SClassViewer::Populate()
 {
-	TArray<FName> PreviousSelection;
+	TArray<FTopLevelAssetPath> PreviousSelection;
 	{
 		TArray<TSharedPtr<FClassViewerNode>> SelectedItems = GetSelectedItems();
 		if (SelectedItems.Num() > 0)
@@ -2381,7 +2382,7 @@ void SClassViewer::Populate()
 		// Take the package names for the internal only classes and convert them into their UClass
 		for (int i = 0; i < InternalClassNames.Num(); i++)
 		{
-			FName PackageClassName = *InternalClassNames[i].ToString();
+			FTopLevelAssetPath PackageClassName(InternalClassNames[i].ToString());
 			const TSharedPtr<FClassViewerNode> ClassNode = ClassViewer::Helpers::ClassHierarchy->FindNodeByGeneratedClassPath(ClassViewer::Helpers::ClassHierarchy->GetObjectRootNode(), PackageClassName);
 
 			if (ClassNode.IsValid())
diff --git a/Engine/Source/Editor/ClassViewer/Private/UnloadedBlueprintData.cpp b/Engine/Source/Editor/ClassViewer/Private/UnloadedBlueprintData.cpp
--- a/Engine/Source/Editor/ClassViewer/Private/UnloadedBlueprintData.cpp
+++ b/Engine/Source/Editor/ClassViewer/Private/UnloadedBlueprintData.cpp
@@ -131,12 +131,22 @@ FName FUnloadedBlueprintData::GetClassPath() const
 {
 	if (ClassViewerNode.IsValid())
 	{
-		return ClassViewerNode.Pin()->ClassPath;
+		return FName(*ClassViewerNode.Pin()->ClassPath.ToString());
 	}
 
 	return FName();
 }
 
+FTopLevelAssetPath FUnloadedBlueprintData::GetClassPathName() const
+{
+	if (ClassViewerNode.IsValid())
+	{
+		return ClassViewerNode.Pin()->ClassPath;
+	}
+
+	return FTopLevelAssetPath();
+}
+
 const TWeakPtr<FClassViewerNode>& FUnloadedBlueprintData::GetClassViewerNode() const
 {
 	return ClassViewerNode;
diff --git a/Engine/Source/Editor/ClassViewer/Private/UnloadedBlueprintData.h b/Engine/Source/Editor/ClassViewer/Private/UnloadedBlueprintData.h
--- a/Engine/Source/Editor/ClassViewer/Private/UnloadedBlueprintData.h
+++ b/Engine/Source/Editor/ClassViewer/Private/UnloadedBlueprintData.h
@@ -30,8 +30,11 @@ public:
 
 	virtual TSharedPtr<FString> GetClassName() const override;
 
+	UE_DEPRECATED(5.1, "Class names are now represented by path names. Please use GetClassPathName.")
 	virtual FName GetClassPath() const override;
 
+	virtual FTopLevelAssetPath GetClassPathName() const override;
+
 	virtual const UClass* GetClassWithin() const override;
 
 	virtual const UClass* GetNativeParent() const override;
diff --git a/Engine/Source/Editor/ClassViewer/Public/ClassViewerFilter.h b/Engine/Source/Editor/ClassViewer/Public/ClassViewerFilter.h
--- a/Engine/Source/Editor/ClassViewer/Public/ClassViewerFilter.h
+++ b/Engine/Source/Editor/ClassViewer/Public/ClassViewerFilter.h
@@ -345,5 +345,11 @@ public:
 	/**
 	 * Get the class path of this blueprint.
 	 */
+	UE_DEPRECATED(5.1, "Class names are now represented by path names. Please use GetClassPathName.")
 	virtual FName GetClassPath() const = 0;
+
+	/**
+	 * Get the class path of this blueprint.
+	 */
+	virtual FTopLevelAssetPath GetClassPathName() const = 0;
 };
diff --git a/Engine/Source/Editor/ClothPainter/Private/ClothingAssetExporter.cpp b/Engine/Source/Editor/ClothPainter/Private/ClothingAssetExporter.cpp
--- a/Engine/Source/Editor/ClothPainter/Private/ClothingAssetExporter.cpp
+++ b/Engine/Source/Editor/ClothPainter/Private/ClothingAssetExporter.cpp
@@ -48,7 +48,7 @@ void ExportClothingAsset(const UClothingAssetBase* ClothingAsset, UClass* Export
 		const FString PackageName = ClothingAsset->GetOutermost()->GetName();
 		ExportClothingAssetDialogConfig.DefaultPath = FPackageName::GetLongPackagePath(PackageName);
 		ExportClothingAssetDialogConfig.DefaultAssetName = ClothingAsset->GetName();
-		ExportClothingAssetDialogConfig.AssetClassNames.Add(ExportedType->GetFName());
+		ExportClothingAssetDialogConfig.AssetClassNames.Add(ExportedType->GetClassPathName());
 		ExportClothingAssetDialogConfig.ExistingAssetPolicy = ESaveAssetDialogExistingAssetPolicy::AllowButWarn;
 		ExportClothingAssetDialogConfig.DialogTitleOverride = LOCTEXT("ExportClothingAssetDialogTitle", "Export Clothing Asset As");
 	}
diff --git a/Engine/Source/Editor/ClothPainter/Private/SClothAssetSelector.cpp b/Engine/Source/Editor/ClothPainter/Private/SClothAssetSelector.cpp
--- a/Engine/Source/Editor/ClothPainter/Private/SClothAssetSelector.cpp
+++ b/Engine/Source/Editor/ClothPainter/Private/SClothAssetSelector.cpp
@@ -1238,7 +1238,7 @@ TSharedRef<SWidget> SClothAssetSelector::OnGenerateSkeletalMeshPickerForClothCop
 	FContentBrowserModule& ContentBrowserModule = FModuleManager::Get().LoadModuleChecked<FContentBrowserModule>(TEXT("ContentBrowser"));
 
 	FAssetPickerConfig AssetPickerConfig;
-	AssetPickerConfig.Filter.ClassNames.Add(USkeletalMesh::StaticClass()->GetFName());
+	AssetPickerConfig.Filter.ClassPaths.Add(USkeletalMesh::StaticClass()->GetClassPathName());
 	AssetPickerConfig.OnAssetSelected = FOnAssetSelected::CreateSP(this, &SClothAssetSelector::OnCopyClothingAssetSelected);
 	AssetPickerConfig.bAllowNullSelection = true;
 	AssetPickerConfig.InitialAssetViewType = EAssetViewType::List;
diff --git a/Engine/Source/Editor/ContentBrowser/Private/CollectionAssetRegistryBridge.cpp b/Engine/Source/Editor/ContentBrowser/Private/CollectionAssetRegistryBridge.cpp
--- a/Engine/Source/Editor/ContentBrowser/Private/CollectionAssetRegistryBridge.cpp
+++ b/Engine/Source/Editor/ContentBrowser/Private/CollectionAssetRegistryBridge.cpp
@@ -34,18 +34,18 @@ public:
 			if (!GIsSavingPackage)
 			{
 				const FString ClassPathStr = InObjectPath.ToString();
-
-				UClass* FoundClass = FindObject<UClass>(ANY_PACKAGE, *ClassPathStr);
+				check(FPackageName::IsValidObjectPath(ClassPathStr));
+				UClass* FoundClass = FindObject<UClass>(nullptr, *ClassPathStr);
 				if (!FoundClass)
 				{
 					// Use the linker to search for class name redirects (from the loaded ActiveClassRedirects)
-					const FString ClassName = FPackageName::ObjectPathToObjectName(ClassPathStr);
-					const FName NewClassName = FLinkerLoad::FindNewNameForClass(*ClassName, false);
+					const FString NewClassName = FLinkerLoad::FindNewPathNameForClass(ClassPathStr, false);
 
-					if (!NewClassName.IsNone())
+					if (!NewClassName.IsEmpty())
 					{
+						check(FPackageName::IsValidObjectPath(NewClassName));
 						// Our new class name might be lacking the path, so try and find it so we can use the full path in the collection
-						FoundClass = FindObject<UClass>(ANY_PACKAGE, *NewClassName.ToString());
+						FoundClass = FindObject<UClass>(nullptr, *NewClassName);
 						if (FoundClass)
 						{
 							OutNewObjectPath = *FoundClass->GetPathName();
diff --git a/Engine/Source/Editor/ContentBrowser/Private/ContentBrowserUtils.cpp b/Engine/Source/Editor/ContentBrowser/Private/ContentBrowserUtils.cpp
--- a/Engine/Source/Editor/ContentBrowser/Private/ContentBrowserUtils.cpp
+++ b/Engine/Source/Editor/ContentBrowser/Private/ContentBrowserUtils.cpp
@@ -459,9 +459,10 @@ void ContentBrowserUtils::CountItemTypes(const TArray<FAssetData>& InItems, int3
 	OutNumAssetItems = 0;
 	OutNumClassItems = 0;
 
+	const FTopLevelAssetPath ClassPath(TEXT("/Script/CoreUObject"), TEXT("Class"));
 	for(const FAssetData& Item : InItems)
 	{
-		if(Item.AssetClass == NAME_Class)
+		if(Item.AssetClassPath == ClassPath)
 		{
 			++OutNumClassItems;
 		}
@@ -573,7 +574,7 @@ void ContentBrowserUtils::ConvertLegacySelectionToVirtualPaths(TArrayView<const
 	ConvertLegacySelectionToVirtualPathsImpl(InAssets, InFolders, InUseFolderPaths, OutVirtualPaths);
 }
 
-void ContentBrowserUtils::AppendAssetFilterToContentBrowserFilter(const FARFilter& InAssetFilter, const TSharedPtr<FNamePermissionList>& InAssetClassPermissionList, const TSharedPtr<FPathPermissionList>& InFolderPermissionList, FContentBrowserDataFilter& OutDataFilter)
+void ContentBrowserUtils::AppendAssetFilterToContentBrowserFilter(const FARFilter& InAssetFilter, const TSharedPtr<FPathPermissionList>& InAssetClassPermissionList, const TSharedPtr<FPathPermissionList>& InFolderPermissionList, FContentBrowserDataFilter& OutDataFilter)
 {
 	if (InAssetFilter.ObjectPaths.Num() > 0 || InAssetFilter.TagsAndValues.Num() > 0 || InAssetFilter.bIncludeOnlyOnDiskAssets)
 	{
@@ -592,14 +593,20 @@ void ContentBrowserUtils::AppendAssetFilterToContentBrowserFilter(const FARFilte
 		PackageFilter.PathPermissionList = InFolderPermissionList;
 	}
 
-	if (InAssetFilter.ClassNames.Num() > 0 || (InAssetClassPermissionList && InAssetClassPermissionList->HasFiltering()))
+	if (InAssetFilter.ClassPaths.Num() > 0 || (InAssetClassPermissionList && InAssetClassPermissionList->HasFiltering()))
 	{
 		FContentBrowserDataClassFilter& ClassFilter = OutDataFilter.ExtraFilters.FindOrAddFilter<FContentBrowserDataClassFilter>();
-		ClassFilter.ClassNamesToInclude = InAssetFilter.ClassNames;
+		for (FTopLevelAssetPath ClassPathName : InAssetFilter.ClassPaths)
+		{
+			ClassFilter.ClassNamesToInclude.Add(ClassPathName.ToString());
+		}
 		ClassFilter.bRecursiveClassNamesToInclude = InAssetFilter.bRecursiveClasses;
 		if (InAssetFilter.bRecursiveClasses)
 		{
-			ClassFilter.ClassNamesToExclude = InAssetFilter.RecursiveClassesExclusionSet.Array();
+			for (FTopLevelAssetPath ClassPathName : InAssetFilter.RecursiveClassPathsExclusionSet)
+			{
+				ClassFilter.ClassNamesToExclude.Add(ClassPathName.ToString());
+			}
 			ClassFilter.bRecursiveClassNamesToExclude = false;
 		}
 		ClassFilter.ClassPermissionList = InAssetClassPermissionList;
diff --git a/Engine/Source/Editor/ContentBrowser/Private/ContentBrowserUtils.h b/Engine/Source/Editor/ContentBrowser/Private/ContentBrowserUtils.h
--- a/Engine/Source/Editor/ContentBrowser/Private/ContentBrowserUtils.h
+++ b/Engine/Source/Editor/ContentBrowser/Private/ContentBrowserUtils.h
@@ -78,7 +78,7 @@ namespace ContentBrowserUtils
 	void ConvertLegacySelectionToVirtualPaths(TArrayView<const FAssetData> InAssets, TArrayView<const FString> InFolders, const bool InUseFolderPaths, TSet<FName>& OutVirtualPaths);
 
 	/** Append the asset registry filter and permission lists to the content browser data filter */
-	void AppendAssetFilterToContentBrowserFilter(const FARFilter& InAssetFilter, const TSharedPtr<FNamePermissionList>& InAssetClassPermissionList, const TSharedPtr<FPathPermissionList>& InFolderPermissionList, FContentBrowserDataFilter& OutDataFilter);
+	void AppendAssetFilterToContentBrowserFilter(const FARFilter& InAssetFilter, const TSharedPtr<FPathPermissionList>& InAssetClassPermissionList, const TSharedPtr<FPathPermissionList>& InFolderPermissionList, FContentBrowserDataFilter& OutDataFilter);
 
 	/* Combine folder filters into a new filter if either are active */
 	TSharedPtr<FPathPermissionList> GetCombinedFolderPermissionList(const TSharedPtr<FPathPermissionList>& FolderPermissionList, const TSharedPtr<FPathPermissionList>& WritableFolderPermissionList);
diff --git a/Engine/Source/Editor/ContentBrowser/Private/FrontendFilters.cpp b/Engine/Source/Editor/ContentBrowser/Private/FrontendFilters.cpp
--- a/Engine/Source/Editor/ContentBrowser/Private/FrontendFilters.cpp
+++ b/Engine/Source/Editor/ContentBrowser/Private/FrontendFilters.cpp
@@ -1132,7 +1132,7 @@ FFrontendFilter_ShowRedirectors::FFrontendFilter_ShowRedirectors(TSharedPtr<FFro
 	: FFrontendFilter(InCategory)
 {
 	bAreRedirectorsInBaseFilter = false;
-	RedirectorClassName = UObjectRedirector::StaticClass()->GetFName();
+	RedirectorClassName = UObjectRedirector::StaticClass()->GetPathName();
 }
 
 void FFrontendFilter_ShowRedirectors::SetCurrentFilter(TArrayView<const FName> InSourcePaths, const FContentBrowserDataFilter& InBaseFilter)
@@ -1147,7 +1147,7 @@ bool FFrontendFilter_ShowRedirectors::PassesFilter(FAssetFilterType InItem) cons
 	if ( !bAreRedirectorsInBaseFilter )
 	{
 		const FContentBrowserItemDataAttributeValue ClassValue = InItem.GetItemAttribute(NAME_Class);
-		return !ClassValue.IsValid() || ClassValue.GetValue<FName>() != RedirectorClassName;
+		return !ClassValue.IsValid() || ClassValue.GetValue<FString>() != RedirectorClassName;
 	}
 
 	return true;
@@ -1339,8 +1339,8 @@ void FFrontendFilter_UsedInAnyLevel::ActiveStateChanged(bool bActive)
 	{
 		// Find all the levels & external actors
 		FARFilter Filter;
-		Filter.ClassNames.Add(UWorld::StaticClass()->GetFName());
-		Filter.ClassNames.Add(AActor::StaticClass()->GetFName());
+		Filter.ClassPaths.Add(UWorld::StaticClass()->GetClassPathName());
+		Filter.ClassPaths.Add(AActor::StaticClass()->GetClassPathName());
 		Filter.bRecursiveClasses = true;
 		FrontendFilterHelper::GetDependencies(Filter, *AssetRegistry, LevelsDependencies);
 	}
@@ -1383,8 +1383,8 @@ void FFrontendFilter_NotUsedInAnyLevel::ActiveStateChanged(bool bActive)
 	{
 		// Find all the levels & external actors
 		FARFilter Filter;
-		Filter.ClassNames.Add(UWorld::StaticClass()->GetFName());
-		Filter.ClassNames.Add(AActor::StaticClass()->GetFName());
+		Filter.ClassPaths.Add(UWorld::StaticClass()->GetClassPathName());
+		Filter.ClassPaths.Add(AActor::StaticClass()->GetClassPathName());
 		Filter.bRecursiveClasses = true;
 		FrontendFilterHelper::GetDependencies(Filter, *AssetRegistry, LevelsDependencies);
 	}
diff --git a/Engine/Source/Editor/ContentBrowser/Private/SAssetDialog.cpp b/Engine/Source/Editor/ContentBrowser/Private/SAssetDialog.cpp
--- a/Engine/Source/Editor/ContentBrowser/Private/SAssetDialog.cpp
+++ b/Engine/Source/Editor/ContentBrowser/Private/SAssetDialog.cpp
@@ -81,7 +81,7 @@ void SAssetDialog::Construct(const FArguments& InArgs, const FSharedAssetDialogC
 	PathPickerConfig.bOnPathSelectedPassesVirtualPaths = true;
 
 	FAssetPickerConfig AssetPickerConfig;
-	AssetPickerConfig.Filter.ClassNames.Append(AssetClassNames);
+	AssetPickerConfig.Filter.ClassPaths.Append(AssetClassNames);
 	AssetPickerConfig.bAllowDragging = false;
 	AssetPickerConfig.InitialAssetViewType = EAssetViewType::Tile;
 	AssetPickerConfig.OnAssetSelected = FOnAssetSelected::CreateSP(this, &SAssetDialog::OnAssetSelected);
@@ -765,7 +765,7 @@ void SAssetDialog::HandlePathSelected(const FString& NewPath)
 
 	FARFilter NewFilter;
 
-	NewFilter.ClassNames.Append(AssetClassNames);
+	NewFilter.ClassPaths.Append(AssetClassNames);
 	NewFilter.PackagePaths.Add(*NewPath);
 
 	SetCurrentlySelectedPath(ConvertedPath.ToString(), ConvertedPathType);
@@ -947,8 +947,8 @@ void SAssetDialog::UpdateInputValidity()
 			FText ErrorMessage;
 			const bool bAllowExistingAsset = (ExistingAssetPolicy == ESaveAssetDialogExistingAssetPolicy::AllowButWarn);
 
-			FName AssetClassName = AssetClassNames.Num() == 1 ? AssetClassNames[0] : NAME_None;
-			UClass* AssetClass = AssetClassName != NAME_None ? FindObject<UClass>(ANY_PACKAGE, *AssetClassName.ToString(), true) : nullptr;
+			FTopLevelAssetPath AssetClassName = AssetClassNames.Num() == 1 ? AssetClassNames[0] : FTopLevelAssetPath();
+			UClass* AssetClass = !AssetClassName.IsNull() ? FindObject<UClass>(AssetClassName, true) : nullptr;
 
 			if ( !ContentBrowserUtils::IsValidObjectPathForCreate(ObjectPath, AssetClass, ErrorMessage, bAllowExistingAsset) )
 			{
@@ -959,10 +959,10 @@ void SAssetDialog::UpdateInputValidity()
 			{
 				FAssetRegistryModule& AssetRegistryModule = FModuleManager::LoadModuleChecked<FAssetRegistryModule>("AssetRegistry");
 				FAssetData ExistingAsset = AssetRegistryModule.Get().GetAssetByObjectPath(FName(*ObjectPath));
-				if (ExistingAsset.IsValid() && !AssetClassNames.Contains(ExistingAsset.AssetClass))
+				if (ExistingAsset.IsValid() && !AssetClassNames.Contains(ExistingAsset.AssetClassPath))
 				{
 					const FString ObjectName = FPackageName::ObjectPathToObjectName(ObjectPath);
-					LastInputValidityErrorText = FText::Format(LOCTEXT("AssetDialog_AssetAlreadyExists", "An asset of type '{0}' already exists at this location with the name '{1}'."), FText::FromName(ExistingAsset.AssetClass), FText::FromString(ObjectName));
+					LastInputValidityErrorText = FText::Format(LOCTEXT("AssetDialog_AssetAlreadyExists", "An asset of type '{0}' already exists at this location with the name '{1}'."), FText::FromString(ExistingAsset.AssetClassPath.ToString()), FText::FromString(ObjectName));
 					bLastInputValidityCheckSuccessful = false;
 				}
 			}
@@ -1031,7 +1031,7 @@ void SAssetDialog::CommitObjectPathForSave()
 			{
 				FAssetRegistryModule& AssetRegistryModule = FModuleManager::LoadModuleChecked<FAssetRegistryModule>("AssetRegistry");
 				FAssetData ExistingAsset = AssetRegistryModule.Get().GetAssetByObjectPath(FName(*ObjectPath));
-				if ( ExistingAsset.IsValid() && AssetClassNames.Contains(ExistingAsset.AssetClass) )
+				if ( ExistingAsset.IsValid() && AssetClassNames.Contains(ExistingAsset.AssetClassPath) )
 				{
 					EAppReturnType::Type ShouldReplace = FMessageDialog::Open( EAppMsgType::YesNo, FText::Format(LOCTEXT("ReplaceAssetMessage", "{0} already exists. Do you want to replace it?"), FText::FromString(CurrentlyEnteredAssetName)) );
 					bProceedWithSave = (ShouldReplace == EAppReturnType::Yes);
diff --git a/Engine/Source/Editor/ContentBrowser/Private/SAssetDialog.h b/Engine/Source/Editor/ContentBrowser/Private/SAssetDialog.h
--- a/Engine/Source/Editor/ContentBrowser/Private/SAssetDialog.h
+++ b/Engine/Source/Editor/ContentBrowser/Private/SAssetDialog.h
@@ -161,7 +161,7 @@ private:
 	FGetCurrentSelectionDelegate GetCurrentSelectionDelegate;
 
 	/** Only assets of these classes will show up */
-	TArray<FName> AssetClassNames;
+	TArray<FTopLevelAssetPath> AssetClassNames;
 
 	/** Fired when assets are chosen for open. Only fired in open dialogs. */
 	FOnAssetsChosenForOpen OnAssetsChosenForOpen;
diff --git a/Engine/Source/Editor/ContentBrowser/Private/SAssetPicker.cpp b/Engine/Source/Editor/ContentBrowser/Private/SAssetPicker.cpp
--- a/Engine/Source/Editor/ContentBrowser/Private/SAssetPicker.cpp
+++ b/Engine/Source/Editor/ContentBrowser/Private/SAssetPicker.cpp
@@ -218,10 +218,10 @@ void SAssetPicker::Construct( const FArguments& InArgs )
 	{
 		// Filters
 		TArray<UClass*> FilterClassList;
-		for(auto Iter = CurrentBackendFilter.ClassNames.CreateIterator(); Iter; ++Iter)
+		for(auto Iter = CurrentBackendFilter.ClassPaths.CreateIterator(); Iter; ++Iter)
 		{
-			FName ClassName = (*Iter);
-			UClass* FilterClass = FindObject<UClass>(ANY_PACKAGE, *ClassName.ToString());
+			FTopLevelAssetPath ClassName = (*Iter);
+			UClass* FilterClass = FindObject<UClass>(ClassName);
 			if(FilterClass)
 			{
 				FilterClassList.AddUnique(FilterClass);
@@ -353,7 +353,7 @@ void SAssetPicker::Construct( const FArguments& InArgs )
 	if (AssetViewPtr.IsValid() && !InArgs._AssetPickerConfig.bAutohideSearchBar)
 	{
 		TextFilter = MakeShareable(new FFrontendFilter_Text());
-		bool bClassNamesProvided = (InArgs._AssetPickerConfig.Filter.ClassNames.Num() != 1);
+		bool bClassNamesProvided = (InArgs._AssetPickerConfig.Filter.ClassPaths.Num() != 1);
 		TextFilter->SetIncludeClassName(bClassNamesProvided || AssetViewPtr->IsIncludingClassNames());
 		TextFilter->SetIncludeAssetPath(AssetViewPtr->IsIncludingAssetPaths());
 		TextFilter->SetIncludeCollectionNames(AssetViewPtr->IsIncludingCollectionNames());
@@ -494,7 +494,7 @@ void SAssetPicker::SetNewBackendFilter(const FARFilter& NewFilter)
 	// Update the Text filter too, since now class names may no longer matter
 	if (TextFilter.IsValid())
 	{
-		TextFilter->SetIncludeClassName(NewFilter.ClassNames.Num() != 1);
+		TextFilter->SetIncludeClassName(NewFilter.ClassPaths.Num() != 1);
 	}
 
 	OnFilterChanged();
diff --git a/Engine/Source/Editor/ContentBrowser/Private/SAssetView.cpp b/Engine/Source/Editor/ContentBrowser/Private/SAssetView.cpp
--- a/Engine/Source/Editor/ContentBrowser/Private/SAssetView.cpp
+++ b/Engine/Source/Editor/ContentBrowser/Private/SAssetView.cpp
@@ -334,7 +334,7 @@ void SAssetView::Construct( const FArguments& InArgs )
 	OwningContentBrowser = InArgs._OwningContentBrowser;
 
 	FAssetToolsModule& AssetToolsModule = FModuleManager::LoadModuleChecked<FAssetToolsModule>("AssetTools");
-	AssetClassPermissionList = AssetToolsModule.Get().GetAssetClassPermissionList(EAssetClassAction::ViewAsset);
+	AssetClassPermissionList = AssetToolsModule.Get().GetAssetClassPathPermissionList(EAssetClassAction::ViewAsset);
 	FolderPermissionList = AssetToolsModule.Get().GetFolderPermissionList();
 	WritableFolderPermissionList = AssetToolsModule.Get().GetWritableFolderPermissionList();
 
@@ -1536,7 +1536,7 @@ static bool IsValidObjectPath(const FString& Path, FString& OutObjectClassName,
 {
 	if (FPackageName::ParseExportTextPath(Path, &OutObjectClassName, &OutObjectPath))
 	{
-		if (UClass* ObjectClass = FindObject<UClass>(ANY_PACKAGE, *OutObjectClassName, true))
+		if (UClass* ObjectClass = UClass::TryFindTypeSlow<UClass>(OutObjectClassName, EFindFirstObjectOptions::ExactClass))
 		{
 			OutPackageName = FPackageName::ObjectPathToPackageName(OutObjectPath);
 			if (FPackageName::IsValidLongPackageName(OutPackageName))
diff --git a/Engine/Source/Editor/ContentBrowser/Public/FrontendFilters.h b/Engine/Source/Editor/ContentBrowser/Public/FrontendFilters.h
--- a/Engine/Source/Editor/ContentBrowser/Public/FrontendFilters.h
+++ b/Engine/Source/Editor/ContentBrowser/Public/FrontendFilters.h
@@ -264,7 +264,7 @@ public:
 
 private:
 	bool bAreRedirectorsInBaseFilter;
-	FName RedirectorClassName;
+	FString RedirectorClassName;
 };
 
 /** A filter that only displays assets used by loaded levels */
diff --git a/Engine/Source/Editor/ContentBrowser/Public/IContentBrowserSingleton.h b/Engine/Source/Editor/ContentBrowser/Public/IContentBrowserSingleton.h
--- a/Engine/Source/Editor/ContentBrowser/Public/IContentBrowserSingleton.h
+++ b/Engine/Source/Editor/ContentBrowser/Public/IContentBrowserSingleton.h
@@ -412,7 +412,7 @@ struct FSharedAssetDialogConfig
 {
 	FText DialogTitleOverride;
 	FString DefaultPath;
-	TArray<FName> AssetClassNames;
+	TArray<FTopLevelAssetPath> AssetClassNames;
 	FVector2D WindowSizeOverride;
 	FOnPathSelected OnPathSelected;
 
diff --git a/Engine/Source/Editor/ContentBrowser/Public/SAssetView.h b/Engine/Source/Editor/ContentBrowser/Public/SAssetView.h
--- a/Engine/Source/Editor/ContentBrowser/Public/SAssetView.h
+++ b/Engine/Source/Editor/ContentBrowser/Public/SAssetView.h
@@ -855,7 +855,7 @@ private:
 	/** The current base source filter for the view */
 	FSourcesData SourcesData;
 	FARFilter BackendFilter;
-	TSharedPtr<FNamePermissionList> AssetClassPermissionList;
+	TSharedPtr<FPathPermissionList> AssetClassPermissionList;
 	TSharedPtr<FPathPermissionList> FolderPermissionList;
 	TSharedPtr<FPathPermissionList> WritableFolderPermissionList;
 	TSharedPtr<FAssetFilterCollectionType> FrontendFilters;
diff --git a/Engine/Source/Editor/ContentBrowserData/Public/ContentBrowserDataFilter.h b/Engine/Source/Editor/ContentBrowserData/Public/ContentBrowserDataFilter.h
--- a/Engine/Source/Editor/ContentBrowserData/Public/ContentBrowserDataFilter.h
+++ b/Engine/Source/Editor/ContentBrowserData/Public/ContentBrowserDataFilter.h
@@ -263,11 +263,11 @@ struct CONTENTBROWSERDATA_API FContentBrowserDataClassFilter
 public:
 	/** Array of class names that should be included in this query */
 	UPROPERTY(EditAnywhere, BlueprintReadWrite, Category="ContentBrowser")
-	TArray<FName> ClassNamesToInclude;
+	TArray<FString> ClassNamesToInclude;
 
 	/** Array of class names that should be excluded from this query */
 	UPROPERTY(EditAnywhere, BlueprintReadWrite, Category="ContentBrowser")
-	TArray<FName> ClassNamesToExclude;
+	TArray<FString> ClassNamesToExclude;
 
 	/** Whether we should include inclusive sub-classes in this query */
 	UPROPERTY(EditAnywhere, BlueprintReadWrite, Category="ContentBrowser")
@@ -278,7 +278,7 @@ public:
 	bool bRecursiveClassNamesToExclude = false;
 
 	/** Optional set of additional class filtering */
-	TSharedPtr<FNamePermissionList> ClassPermissionList;
+	TSharedPtr<FPathPermissionList> ClassPermissionList;
 };
 
 /**
diff --git a/Engine/Source/Editor/DataLayerEditor/Private/DataLayer/DataLayerMode.cpp b/Engine/Source/Editor/DataLayerEditor/Private/DataLayer/DataLayerMode.cpp
--- a/Engine/Source/Editor/DataLayerEditor/Private/DataLayer/DataLayerMode.cpp
+++ b/Engine/Source/Editor/DataLayerEditor/Private/DataLayer/DataLayerMode.cpp
@@ -1700,7 +1700,7 @@ UDataLayerAsset* FDataLayerMode::PromptDataLayerAssetSelection()
 
 	FOpenAssetDialogConfig Config;
 	Config.bAllowMultipleSelection = false;
-	Config.AssetClassNames.Add(UDataLayerAsset::StaticClass()->GetFName());
+	Config.AssetClassNames.Add(UDataLayerAsset::StaticClass()->GetClassPathName());
 	Config.DefaultPath = PickDataLayerDialogPath;
 	Config.DialogTitleOverride = LOCTEXT("PickDataLayerAssetDialogTitle", "Pick a Data Layer Asset");
 
diff --git a/Engine/Source/Editor/DataTableEditor/Private/DataTableEditor.cpp b/Engine/Source/Editor/DataTableEditor/Private/DataTableEditor.cpp
--- a/Engine/Source/Editor/DataTableEditor/Private/DataTableEditor.cpp
+++ b/Engine/Source/Editor/DataTableEditor/Private/DataTableEditor.cpp
@@ -966,7 +966,7 @@ void FDataTableEditor::PostRegenerateMenusAndToolbars()
 				SNew(SHyperlink)
 				.Style(FAppStyle::Get(), "Common.GotoNativeCodeHyperlink")
 				.OnNavigate(this, &FDataTableEditor::OnEditDataTableStructClicked)
-				.Text(FText::FromName(DataTable->GetRowStructName()))
+				.Text(FText::FromName(DataTable->GetRowStructPathName().GetAssetName()))
 				.ToolTipText(LOCTEXT("DataTableRowToolTip", "Open the struct used for each row in this data table"))
 			]
 			+ SHorizontalBox::Slot()
diff --git a/Engine/Source/Editor/DetailCustomizations/Private/ComponentReferenceCustomization.cpp b/Engine/Source/Editor/DetailCustomizations/Private/ComponentReferenceCustomization.cpp
--- a/Engine/Source/Editor/DetailCustomizations/Private/ComponentReferenceCustomization.cpp
+++ b/Engine/Source/Editor/DetailCustomizations/Private/ComponentReferenceCustomization.cpp
@@ -147,7 +147,7 @@ void FComponentReferenceCustomization::BuildClassFilters()
 
 			for (const FString& ClassName : ClassFilterNames)
 			{
-				UClass* Class = FindObject<UClass>(ANY_PACKAGE, *ClassName);
+				UClass* Class = UClass::TryFindTypeSlow<UClass>(ClassName);
 				if (!Class)
 				{
 					Class = LoadObject<UClass>(nullptr, *ClassName);
diff --git a/Engine/Source/Editor/DetailCustomizations/Private/DataTableCustomization.cpp b/Engine/Source/Editor/DetailCustomizations/Private/DataTableCustomization.cpp
--- a/Engine/Source/Editor/DetailCustomizations/Private/DataTableCustomization.cpp
+++ b/Engine/Source/Editor/DetailCustomizations/Private/DataTableCustomization.cpp
@@ -19,7 +19,7 @@ void FDataTableCustomizationLayout::CustomizeHeader(TSharedRef<class IPropertyHa
 	{
 		const FString& RowType = StructPropertyHandle->GetMetaData(TEXT("RowType"));
 		RowTypeFilter = FName(*RowType);
-		RowFilterStruct = FindObject<UScriptStruct>(ANY_PACKAGE, *RowTypeFilter.ToString());
+		RowFilterStruct = UClass::TryFindTypeSlow<UScriptStruct>(RowType);
 	}
 
 	FSimpleDelegate OnDataTableChangedDelegate = FSimpleDelegate::CreateSP(this, &FDataTableCustomizationLayout::OnDataTableChanged);
@@ -191,16 +191,16 @@ bool FDataTableCustomizationLayout::ShouldFilterAsset(const struct FAssetData& A
 	if (!RowTypeFilter.IsNone())
 	{
 		static const FName RowStructureTagName("RowStructure");
-		FName RowStructure;
-		if (AssetData.GetTagValue<FName>(RowStructureTagName, RowStructure))
+		FString RowStructure;
+		if (AssetData.GetTagValue<FString>(RowStructureTagName, RowStructure))
 		{
-			if (RowStructure == RowTypeFilter)
+			if (RowStructure == RowTypeFilter.ToString())
 			{
 				return false;
 			}
 
 			// This is slow, but at the moment we don't have an alternative to the short struct name search
-			UScriptStruct* RowStruct = FindObject<UScriptStruct>(ANY_PACKAGE, *RowStructure.ToString());
+			UScriptStruct* RowStruct = UClass::TryFindTypeSlow<UScriptStruct>(RowStructure);
 			if (RowStruct && RowFilterStruct && RowStruct->IsChildOf(RowFilterStruct))
 			{
 				return false;
diff --git a/Engine/Source/Editor/DetailCustomizations/Private/DialogueWaveWidgets.cpp b/Engine/Source/Editor/DetailCustomizations/Private/DialogueWaveWidgets.cpp
--- a/Engine/Source/Editor/DetailCustomizations/Private/DialogueWaveWidgets.cpp
+++ b/Engine/Source/Editor/DetailCustomizations/Private/DialogueWaveWidgets.cpp
@@ -402,14 +402,16 @@ void SDialogueVoicePropertyEditor::OnDialogueVoicePathChanged( const FText& NewT
 			
 		if( !NewString.IsEmpty() )
 		{
-			UObject* Package = ANY_PACKAGE;
 			if( NewString.Contains( TEXT(".") ) )
 			{
 				// Formatted text string, use the exact path instead of any package
-				Package = NULL;
+				DialogueVoiceToAssign = FindObject<UDialogueVoice>(nullptr, *NewString);
 			}
-
-			DialogueVoiceToAssign = Cast<UDialogueVoice>( StaticFindObject( UDialogueVoice::StaticClass(), Package, *NewString ) );		
+			else
+			{
+				DialogueVoiceToAssign = FindFirstObject<UDialogueVoice>(*NewString, EFindFirstObjectOptions::NativeFirst | EFindFirstObjectOptions::EnsureIfAmbiguous);
+			}
+	
 			if( !DialogueVoiceToAssign )
 			{
 				DialogueVoiceToAssign = Cast<UDialogueVoice>( StaticLoadObject( UDialogueVoice::StaticClass(), NULL, *NewString ) );
diff --git a/Engine/Source/Editor/DetailCustomizations/Private/SkeletalMeshComponentDetails.cpp b/Engine/Source/Editor/DetailCustomizations/Private/SkeletalMeshComponentDetails.cpp
--- a/Engine/Source/Editor/DetailCustomizations/Private/SkeletalMeshComponentDetails.cpp
+++ b/Engine/Source/Editor/DetailCustomizations/Private/SkeletalMeshComponentDetails.cpp
@@ -248,7 +248,7 @@ void FSkeletalMeshComponentDetails::UpdateSkeletonNameAndPickerVisibility()
 	if (Skeleton)
 	{
 		bAnimPickerEnabled = true;
-		SelectedSkeletonName = FString::Printf(TEXT("%s'%s'"), *Skeleton->GetClass()->GetName(), *Skeleton->GetPathName());
+		SelectedSkeletonName = FObjectPropertyBase::GetExportPath(Skeleton);
 	}
 	else
 	{
diff --git a/Engine/Source/Editor/DetailCustomizations/Private/WorldSettingsDetails.cpp b/Engine/Source/Editor/DetailCustomizations/Private/WorldSettingsDetails.cpp
--- a/Engine/Source/Editor/DetailCustomizations/Private/WorldSettingsDetails.cpp
+++ b/Engine/Source/Editor/DetailCustomizations/Private/WorldSettingsDetails.cpp
@@ -402,7 +402,7 @@ TSharedRef<SWidget> FLightmapCustomNodeBuilder::MakeLightMapList(TSharedPtr<FLig
 					// Class
 					SNew(STextBlock)
 					.Font(IDetailLayoutBuilder::GetDetailFont())
-					.Text(FText::FromName(LightMapAssetData.AssetClass))
+					.Text(FText::FromString(LightMapAssetData.AssetClassPath.ToString()))
 				]
 			]
 		];
diff --git a/Engine/Source/Editor/Documentation/Private/UDNParser.cpp b/Engine/Source/Editor/Documentation/Private/UDNParser.cpp
--- a/Engine/Source/Editor/Documentation/Private/UDNParser.cpp
+++ b/Engine/Source/Editor/Documentation/Private/UDNParser.cpp
@@ -1319,7 +1319,7 @@ bool FUDNParser::ParseAssetLink(FString &InternalLink)
 		FString Action = Token[0];
 		FString AssetName = Token[1];
 
-		UObject* RequiredObject = FindObject<UObject>(ANY_PACKAGE, *AssetName);
+		UObject* RequiredObject = FindFirstObject<UObject>(*AssetName, EFindFirstObjectOptions::NativeFirst | EFindFirstObjectOptions::EnsureIfAmbiguous);
 		if (RequiredObject != nullptr)
 		{
 			if (Action == TEXT("EDIT"))
diff --git a/Engine/Source/Editor/EditorConfig/Private/EditorMetadataOverrides.cpp b/Engine/Source/Editor/EditorConfig/Private/EditorMetadataOverrides.cpp
--- a/Engine/Source/Editor/EditorConfig/Private/EditorMetadataOverrides.cpp
+++ b/Engine/Source/Editor/EditorConfig/Private/EditorMetadataOverrides.cpp
@@ -332,7 +332,7 @@ bool UEditorMetadataOverrides::GetClassMetadata(const FField* Field, FName Key,
 		return false;
 	}
 
-	OutValue = FindObject<UClass>(ANY_PACKAGE, *ClassName);
+	OutValue = UClass::TryFindTypeSlow<UClass>(ClassName);
 	return true; // we return true here even if the value is null because we did have a value, it just wasn't a valid class name
 }
 
@@ -664,7 +664,7 @@ bool UEditorMetadataOverrides::GetClassMetadata(const UStruct* Struct, FName Key
 		return false;
 	}
 
-	OutValue = FindObject<UClass>(ANY_PACKAGE, *ClassName);
+	OutValue = UClass::TryFindTypeSlow<UClass>(ClassName);
 	return true; // we return true here even if the class is null because we did have a value, it just wasn't a valid class name
 }
 
diff --git a/Engine/Source/Editor/EditorWidgets/Private/Filters/CustomClassFilterData.cpp b/Engine/Source/Editor/EditorWidgets/Private/Filters/CustomClassFilterData.cpp
--- a/Engine/Source/Editor/EditorWidgets/Private/Filters/CustomClassFilterData.cpp
+++ b/Engine/Source/Editor/EditorWidgets/Private/Filters/CustomClassFilterData.cpp
@@ -55,7 +55,7 @@ void FCustomClassFilterData::BuildBackendFilter(FARFilter &Filter)
 	// If there is no AssetTypeAction for this filter, simply add the class name to the FARFilter
 	else
 	{
-		Filter.ClassNames.Add(Class->GetFName());
+		Filter.ClassPaths.Add(Class->GetClassPathName());
 		Filter.bRecursiveClasses = true;
 	}
 }
diff --git a/Engine/Source/Editor/EditorWidgets/Private/STextPropertyEditableTextBox.cpp b/Engine/Source/Editor/EditorWidgets/Private/STextPropertyEditableTextBox.cpp
--- a/Engine/Source/Editor/EditorWidgets/Private/STextPropertyEditableTextBox.cpp
+++ b/Engine/Source/Editor/EditorWidgets/Private/STextPropertyEditableTextBox.cpp
@@ -375,7 +375,7 @@ void STextPropertyEditableStringTableReference::UpdateStringTableComboOptions()
 		FAssetRegistryModule& AssetRegistryModule = FModuleManager::Get().LoadModuleChecked<FAssetRegistryModule>(AssetRegistryConstants::ModuleName);
 
 		TArray<FAssetData> StringTableAssets;
-		AssetRegistryModule.Get().GetAssetsByClass(UStringTable::StaticClass()->GetFName(), StringTableAssets);
+		AssetRegistryModule.Get().GetAssetsByClass(UStringTable::StaticClass()->GetClassPathName(), StringTableAssets);
 
 		for (const FAssetData& StringTableAsset : StringTableAssets)
 		{
diff --git a/Engine/Source/Editor/EditorWidgets/Public/Filters/SAssetFilterBar.h b/Engine/Source/Editor/EditorWidgets/Public/Filters/SAssetFilterBar.h
--- a/Engine/Source/Editor/EditorWidgets/Public/Filters/SAssetFilterBar.h
+++ b/Engine/Source/Editor/EditorWidgets/Public/Filters/SAssetFilterBar.h
@@ -274,18 +274,18 @@ public:
 			for(const TSharedRef<FCustomClassFilterData> &CustomClassFilter : CustomClassFilters)
 			{
 				const UClass* TypeClass = CustomClassFilter->GetClass();
-				if (TypeClass && !CombinedFilter.ClassNames.Contains(TypeClass->GetFName()))
+				if (TypeClass && !CombinedFilter.ClassPaths.Contains(TypeClass->GetClassPathName()))
 				{
-					CombinedFilter.RecursiveClassesExclusionSet.Add(TypeClass->GetFName());
+					CombinedFilter.RecursiveClassPathsExclusionSet.Add(TypeClass->GetClassPathName());
 				}
 			}
 		}
 
 		// HACK: A blueprint can be shown as Blueprint or as BlueprintGeneratedClass, but we don't want to distinguish them while filtering.
 		// This should be removed, once all blueprints are shown as BlueprintGeneratedClass.
-		if(CombinedFilter.ClassNames.Contains(FName(TEXT("Blueprint"))))
+		if(CombinedFilter.ClassPaths.Contains(FTopLevelAssetPath(TEXT("/Script/Engine"), TEXT("Blueprint"))))
 		{
-			CombinedFilter.ClassNames.AddUnique(FName(TEXT("BlueprintGeneratedClass")));
+			CombinedFilter.ClassPaths.AddUnique(FTopLevelAssetPath(TEXT("/Script/Engine"), TEXT("BlueprintGeneratedClass")));
 		}
 
 		return CombinedFilter;
@@ -522,7 +522,7 @@ protected:
 		};
 		AssetTypeActionsList.Sort( FCompareIAssetTypeActions() );
 
-		const TSharedRef<FNamePermissionList>& AssetClassPermissionList = AssetToolsModule.Get().GetAssetClassPermissionList(EAssetClassAction::CreateAsset);
+		const TSharedRef<FPathPermissionList>& AssetClassPermissionList = AssetToolsModule.Get().GetAssetClassPathPermissionList(EAssetClassAction::CreateAsset);
 
 		// For every asset type, convert it to an FCustomClassFilterData and add it to the list
 		for (int32 ClassIdx = 0; ClassIdx < AssetTypeActionsList.Num(); ++ClassIdx)
@@ -534,7 +534,7 @@ protected:
 				if ( ensure(TypeActions.IsValid()) && TypeActions->CanFilter() )
 				{
 					UClass* SupportedClass = TypeActions->GetSupportedClass();
-					if ((!SupportedClass || AssetClassPermissionList->PassesFilter(SupportedClass->GetFName())))
+					if ((!SupportedClass || AssetClassPermissionList->PassesFilter(SupportedClass->GetPathName())))
 					{
 						// Convert the AssetTypeAction to an FCustomClassFilterData and add it to our list
 						TSharedRef<FCustomClassFilterData> CustomClassFilterData = MakeShared<FCustomClassFilterData>(TypeActions);
diff --git a/Engine/Source/Editor/FoliageEdit/Private/SFoliagePalette.cpp b/Engine/Source/Editor/FoliageEdit/Private/SFoliagePalette.cpp
--- a/Engine/Source/Editor/FoliageEdit/Private/SFoliagePalette.cpp
+++ b/Engine/Source/Editor/FoliageEdit/Private/SFoliagePalette.cpp
@@ -1079,7 +1079,7 @@ void SFoliagePalette::FillReplaceFoliageTypeSubmenu(FMenuBuilder& MenuBuilder)
 	FContentBrowserModule& ContentBrowserModule = FModuleManager::Get().LoadModuleChecked<FContentBrowserModule>(TEXT("ContentBrowser"));
 
 	FAssetPickerConfig AssetPickerConfig;
-	AssetPickerConfig.Filter.ClassNames.Add(UFoliageType::StaticClass()->GetFName());
+	AssetPickerConfig.Filter.ClassPaths.Add(UFoliageType::StaticClass()->GetClassPathName());
 	AssetPickerConfig.Filter.bRecursiveClasses = true;
 	AssetPickerConfig.OnAssetSelected = FOnAssetSelected::CreateSP(this, &SFoliagePalette::OnReplaceFoliageTypeSelected);
 	AssetPickerConfig.InitialAssetViewType = EAssetViewType::List;
diff --git a/Engine/Source/Editor/FontEditor/Private/SCompositeFontEditor.cpp b/Engine/Source/Editor/FontEditor/Private/SCompositeFontEditor.cpp
--- a/Engine/Source/Editor/FontEditor/Private/SCompositeFontEditor.cpp
+++ b/Engine/Source/Editor/FontEditor/Private/SCompositeFontEditor.cpp
@@ -1075,7 +1075,7 @@ UFontFace* STypefaceEntryEditor::SaveFontFaceAsAsset(const UFontFace* InFontFace
 	FSaveAssetDialogConfig SaveAssetDialogConfig;
 	SaveAssetDialogConfig.DefaultPath = DefaultPackagePath;
 	SaveAssetDialogConfig.DefaultAssetName = DefaultFaceAssetName;
-	SaveAssetDialogConfig.AssetClassNames.Add(InFontFace->GetClass()->GetFName());
+	SaveAssetDialogConfig.AssetClassNames.Add(InFontFace->GetClass()->GetClassPathName());
 	SaveAssetDialogConfig.ExistingAssetPolicy = ESaveAssetDialogExistingAssetPolicy::AllowButWarn;
 	SaveAssetDialogConfig.DialogTitleOverride = LOCTEXT("SaveFontFaceDialogTitle", "Save Font Face");
 
diff --git a/Engine/Source/Editor/GameProjectGeneration/Private/GameProjectUtils.cpp b/Engine/Source/Editor/GameProjectGeneration/Private/GameProjectUtils.cpp
--- a/Engine/Source/Editor/GameProjectGeneration/Private/GameProjectUtils.cpp
+++ b/Engine/Source/Editor/GameProjectGeneration/Private/GameProjectUtils.cpp
@@ -1314,7 +1314,7 @@ UTemplateProjectDefs* GameProjectUtils::LoadTemplateDefs(const FString& ProjectD
 		const bool bFoundValue = GConfig->GetString(*UTemplateProjectDefs::StaticClass()->GetPathName(), TEXT("TemplateProjectDefsClass"), ClassName, TemplateDefsIniFilename);
 		if (bFoundValue && ClassName.Len() > 0)
 		{
-			UClass* OverrideClass = FindObject<UClass>(ANY_PACKAGE, *ClassName, false);
+			UClass* OverrideClass = UClass::TryFindTypeSlow<UClass>(ClassName);
 			if (nullptr != OverrideClass)
 			{
 				ClassToConstruct = OverrideClass;
diff --git a/Engine/Source/Editor/GameProjectGeneration/Private/SNewClassDialog.cpp b/Engine/Source/Editor/GameProjectGeneration/Private/SNewClassDialog.cpp
--- a/Engine/Source/Editor/GameProjectGeneration/Private/SNewClassDialog.cpp
+++ b/Engine/Source/Editor/GameProjectGeneration/Private/SNewClassDialog.cpp
@@ -1117,7 +1117,7 @@ void SNewClassDialog::FinishClicked()
 			// @todo show fail reason in error label
 			FMessageDialog::Open(EAppMsgType::Ok, LOCTEXT("AddCodeFailed_Blueprint_NoBase", "No parent class has been specified. Failed to generate new Blueprint class."));
 		}
-		else if (FindObject<UBlueprint>(ANY_PACKAGE, *PackagePath))
+		else if (FindObject<UBlueprint>(nullptr, *PackagePath))
 		{
 			// @todo show fail reason in error label
 			FMessageDialog::Open(EAppMsgType::Ok, LOCTEXT("AddCodeFailed_Blueprint_AlreadyExists", "The chosen Blueprint class already exists, please try again with a different name."));
diff --git a/Engine/Source/Editor/GraphEditor/Private/KismetPins/SGraphPinClass.cpp b/Engine/Source/Editor/GraphEditor/Private/KismetPins/SGraphPinClass.cpp
--- a/Engine/Source/Editor/GraphEditor/Private/KismetPins/SGraphPinClass.cpp
+++ b/Engine/Source/Editor/GraphEditor/Private/KismetPins/SGraphPinClass.cpp
@@ -130,7 +130,7 @@ TSharedRef<SWidget> SGraphPinClass::GenerateAssetPicker()
 		FString PossibleInterface = ParentNode->GetPinMetaData(GraphPinObj->PinName, TEXT("MustImplement"));
 		if (!PossibleInterface.IsEmpty())
 		{
-			Filter->RequiredInterface = FindObject<UClass>(ANY_PACKAGE, *PossibleInterface);
+			Filter->RequiredInterface = UClass::TryFindTypeSlow<UClass>(PossibleInterface);
 		}
 	}
 
@@ -222,7 +222,7 @@ const FAssetData& SGraphPinClass::GetAssetData(bool bRuntimePath) const
 				FString ObjectName = FPackageName::ObjectPathToObjectName(EditorPath);
 
 				// Fake one
-				CachedEditorAssetData = FAssetData(FName(*PackageName), FName(*PackagePath), FName(*ObjectName), UObject::StaticClass()->GetFName());
+				CachedEditorAssetData = FAssetData(FName(*PackageName), FName(*PackagePath), FName(*ObjectName), UObject::StaticClass()->GetClassPathName());
 			}
 		}
 	}
diff --git a/Engine/Source/Editor/GraphEditor/Private/KismetPins/SGraphPinObject.cpp b/Engine/Source/Editor/GraphEditor/Private/KismetPins/SGraphPinObject.cpp
--- a/Engine/Source/Editor/GraphEditor/Private/KismetPins/SGraphPinObject.cpp
+++ b/Engine/Source/Editor/GraphEditor/Private/KismetPins/SGraphPinObject.cpp
@@ -200,7 +200,7 @@ TSharedRef<SWidget> SGraphPinObject::GenerateAssetPicker()
 	FContentBrowserModule& ContentBrowserModule = FModuleManager::Get().LoadModuleChecked<FContentBrowserModule>(TEXT("ContentBrowser"));
 
 	FAssetPickerConfig AssetPickerConfig;
-	AssetPickerConfig.Filter.ClassNames.Add(AllowedClass->GetFName());
+	AssetPickerConfig.Filter.ClassPaths.Add(AllowedClass->GetClassPathName());
 	AssetPickerConfig.bAllowNullSelection = true;
 	AssetPickerConfig.Filter.bRecursiveClasses = true;
 	AssetPickerConfig.OnAssetSelected = FOnAssetSelected::CreateSP(this, &SGraphPinObject::OnAssetSelectedFromPicker);
@@ -213,14 +213,15 @@ TSharedRef<SWidget> SGraphPinObject::GenerateAssetPicker()
 	if( !AllowedClassesFilterString.IsEmpty() )
 	{
 		// Clear out the allowed class names and have the pin's metadata override.
-		AssetPickerConfig.Filter.ClassNames.Empty();
+		AssetPickerConfig.Filter.ClassPaths.Empty();
 
 		// Parse and add the classes from the metadata
 		TArray<FString> AllowedClassesFilterNames;
 		AllowedClassesFilterString.ParseIntoArrayWS(AllowedClassesFilterNames, TEXT(","), true);
 		for(const FString& AllowedClassesFilterName : AllowedClassesFilterNames)
 		{
-			AssetPickerConfig.Filter.ClassNames.Add(FName(*AllowedClassesFilterName));
+			ensureAlwaysMsgf(!FPackageName::IsShortPackageName(AllowedClassesFilterName), TEXT("Short class names are not supported as AllowedClasses on pin \"%s\": class \"%s\""), *GraphPinObj->PinName.ToString(), *AllowedClassesFilterName);
+			AssetPickerConfig.Filter.ClassPaths.Add(FTopLevelAssetPath(AllowedClassesFilterName));
 		}
 	}
 
@@ -231,7 +232,8 @@ TSharedRef<SWidget> SGraphPinObject::GenerateAssetPicker()
 		DisallowedClassesFilterString.ParseIntoArrayWS(DisallowedClassesFilterNames, TEXT(","), true);
 		for(const FString& DisallowedClassesFilterName : DisallowedClassesFilterNames)
 		{
-			AssetPickerConfig.Filter.RecursiveClassesExclusionSet.Add(FName(*DisallowedClassesFilterName));
+			ensureAlwaysMsgf(!FPackageName::IsShortPackageName(DisallowedClassesFilterName), TEXT("Short class names are not supported as DisallowedClasses on pin \"%s\": class \"%s\""), *GraphPinObj->PinName.ToString(), *DisallowedClassesFilterName);
+			AssetPickerConfig.Filter.RecursiveClassPathsExclusionSet.Add(FTopLevelAssetPath(DisallowedClassesFilterName));
 		}
 	}
 
@@ -424,7 +426,7 @@ const FAssetData& SGraphPinObject::GetAssetData(bool bRuntimePath) const
 				FString ObjectName = FPackageName::ObjectPathToObjectName(GraphPinObj->DefaultValue);
 
 				// Fake one
-				CachedAssetData = FAssetData(FName(*PackageName), FName(*PackagePath), FName(*ObjectName), UObject::StaticClass()->GetFName());
+				CachedAssetData = FAssetData(FName(*PackageName), FName(*PackagePath), FName(*ObjectName), UObject::StaticClass()->GetClassPathName());
 			}
 		}
 	}
diff --git a/Engine/Source/Editor/Kismet/Private/BlueprintAssetHandler.cpp b/Engine/Source/Editor/Kismet/Private/BlueprintAssetHandler.cpp
--- a/Engine/Source/Editor/Kismet/Private/BlueprintAssetHandler.cpp
+++ b/Engine/Source/Editor/Kismet/Private/BlueprintAssetHandler.cpp
@@ -40,8 +40,8 @@ class FBlueprintAssetTypeHandler : public IBlueprintAssetHandler
 FBlueprintAssetHandler::FBlueprintAssetHandler()
 {
 	// Register default handlers
-	RegisterHandler<FLevelBlueprintAssetHandler>(UWorld::StaticClass()->GetFName());
-	RegisterHandler<FBlueprintAssetTypeHandler>(UBlueprint::StaticClass()->GetFName());
+	RegisterHandler<FLevelBlueprintAssetHandler>(UWorld::StaticClass()->GetClassPathName());
+	RegisterHandler<FBlueprintAssetTypeHandler>(UBlueprint::StaticClass()->GetClassPathName());
 }
 
 FBlueprintAssetHandler& FBlueprintAssetHandler::Get()
@@ -50,7 +50,7 @@ FBlueprintAssetHandler& FBlueprintAssetHandler::Get()
 	return Singleton;
 }
 
-void FBlueprintAssetHandler::RegisterHandler(FName EligibleClass, TUniquePtr<IBlueprintAssetHandler>&& InHandler)
+void FBlueprintAssetHandler::RegisterHandler(FTopLevelAssetPath EligibleClass, TUniquePtr<IBlueprintAssetHandler>&& InHandler)
 {
 	ClassNames.Add(EligibleClass);
 	Handlers.Add(MoveTemp(InHandler));
@@ -61,7 +61,7 @@ const IBlueprintAssetHandler* FBlueprintAssetHandler::FindHandler(const UClass*
 	UClass* StopAtClass = UObject::StaticClass();
 	while (InClass && InClass != StopAtClass)
 	{
-		int32 Index = ClassNames.IndexOfByKey(InClass->GetFName());
+		int32 Index = ClassNames.IndexOfByKey(InClass->GetClassPathName());
 		if (Index != INDEX_NONE)
 		{
 			return Handlers[Index].Get();
diff --git a/Engine/Source/Editor/Kismet/Private/BlueprintDetailsCustomization.cpp b/Engine/Source/Editor/Kismet/Private/BlueprintDetailsCustomization.cpp
--- a/Engine/Source/Editor/Kismet/Private/BlueprintDetailsCustomization.cpp
+++ b/Engine/Source/Editor/Kismet/Private/BlueprintDetailsCustomization.cpp
@@ -5727,7 +5727,7 @@ void FBlueprintInterfaceLayout::GetManagedListItems(TArray<FManagedListItem>& Ou
 			if (const TSubclassOf<UInterface> Interface = ImplementedInterface.Interface)
 			{
 				FManagedListItem ItemDesc;
-				ItemDesc.ItemName = Interface->GetFName().ToString();
+				ItemDesc.ItemName = Interface->GetPathName();
 				ItemDesc.DisplayName = Interface->GetDisplayNameText();
 				ItemDesc.bIsRemovable = true;
 
@@ -5753,7 +5753,7 @@ void FBlueprintInterfaceLayout::GetManagedListItems(TArray<FManagedListItem>& Ou
 				if (CurrentInterface.Class)
 				{
 					FManagedListItem ItemDesc;
-					ItemDesc.ItemName = CurrentInterface.Class->GetFName().ToString();
+					ItemDesc.ItemName = CurrentInterface.Class->GetPathName();
 					ItemDesc.DisplayName = CurrentInterface.Class->GetDisplayNameText();
 					ItemDesc.bIsRemovable = false;
 
@@ -5775,7 +5775,7 @@ void FBlueprintInterfaceLayout::OnRemoveItem(const FManagedListItem& Item)
 		// We canceled!
 		return;
 	}
-	const FName InterfaceFName = FName(*Item.ItemName);
+	const FTopLevelAssetPath InterfacePathName(Item.ItemName);
 
 	UBlueprint* Blueprint = GetBlueprintObjectChecked();
 	TSharedPtr<FBlueprintEditor> BlueprintEditorPtr = GetPinnedBlueprintEditorPtr();
@@ -5783,7 +5783,7 @@ void FBlueprintInterfaceLayout::OnRemoveItem(const FManagedListItem& Item)
 	{
 		// Close all graphs that are about to be removed
 		TArray<UEdGraph*> Graphs;
-		FBlueprintEditorUtils::GetInterfaceGraphs(Blueprint, InterfaceFName, Graphs);
+		FBlueprintEditorUtils::GetInterfaceGraphs(Blueprint, InterfacePathName, Graphs);
 		for (TArray<UEdGraph*>::TIterator GraphIt(Graphs); GraphIt; ++GraphIt)
 		{
 			BlueprintEditorPtr->CloseDocumentTab(*GraphIt);
@@ -5791,7 +5791,7 @@ void FBlueprintInterfaceLayout::OnRemoveItem(const FManagedListItem& Item)
 	}
 
 	// Do the work of actually removing the interface
-	FBlueprintEditorUtils::RemoveInterface(Blueprint, InterfaceFName, DialogReturn == EAppReturnType::Yes);
+	FBlueprintEditorUtils::RemoveInterface(Blueprint, InterfacePathName, DialogReturn == EAppReturnType::Yes);
 
 	RegenerateChildContent();
 
@@ -5809,7 +5809,7 @@ void FBlueprintInterfaceLayout::OnClassPicked(UClass* PickedClass)
 	{
 		UBlueprint* Blueprint = GetBlueprintObjectChecked();
 
-		FBlueprintEditorUtils::ImplementNewInterface(Blueprint, PickedClass->GetFName());
+		FBlueprintEditorUtils::ImplementNewInterface(Blueprint, PickedClass->GetClassPathName());
 
 		RegenerateChildContent();
 	}
diff --git a/Engine/Source/Editor/Kismet/Private/BlueprintEditor.cpp b/Engine/Source/Editor/Kismet/Private/BlueprintEditor.cpp
--- a/Engine/Source/Editor/Kismet/Private/BlueprintEditor.cpp
+++ b/Engine/Source/Editor/Kismet/Private/BlueprintEditor.cpp
@@ -1835,7 +1835,7 @@ struct FLoadObjectsFromAssetRegistryHelper
 		const double CompileStartTime = FPlatformTime::Seconds();
 
 		TArray<FAssetData> AssetData;
-		AssetRegistryModule.Get().GetAssetsByClass(TObjectType::StaticClass()->GetFName(), AssetData);
+		AssetRegistryModule.Get().GetAssetsByClass(TObjectType::StaticClass()->GetClassPathName(), AssetData);
 
 		for (int32 AssetIndex = 0; AssetIndex < AssetData.Num(); ++AssetIndex)
 		{
@@ -1943,7 +1943,7 @@ void FBlueprintEditor::LoadLibrariesFromAssetRegistry()
 
 			// Collect a full list of assets with the specified class
 			TArray<FAssetData> AssetData;
-			AssetRegistryModule.Get().GetAssetsByClass(UBlueprint::StaticClass()->GetFName(), AssetData);
+			AssetRegistryModule.Get().GetAssetsByClass(UBlueprint::StaticClass()->GetClassPathName(), AssetData);
 
 			GWarn->BeginSlowTask(LOCTEXT("LoadingBlueprintAssetData", "Loading Blueprint Asset Data"), true);
 
diff --git a/Engine/Source/Editor/Kismet/Private/BlueprintEditorCommands.cpp b/Engine/Source/Editor/Kismet/Private/BlueprintEditorCommands.cpp
--- a/Engine/Source/Editor/Kismet/Private/BlueprintEditorCommands.cpp
+++ b/Engine/Source/Editor/Kismet/Private/BlueprintEditorCommands.cpp
@@ -312,8 +312,10 @@ void FBlueprintSpawnNodeCommands::RegisterCommands()
 			continue;
 		}
 
+		UE_CLOG(FPackageName::IsShortPackageName(ClassName), LogClass, Warning, TEXT("Short class name found when parsing ini: [%s] %s"), *ConfigSection, *SettingName);
+
 		FString CommandLabel;
-		UClass* FoundClass = FindObject<UClass>(ANY_PACKAGE, *ClassName, true);
+		UClass* FoundClass = FindFirstObject<UClass>(*ClassName, EFindFirstObjectOptions::ExactClass, ELogVerbosity::Warning, TEXT("looking for SpawnNodes"));
 		TSharedPtr< FNodeSpawnInfo > InfoPtr;
 
 		if(FoundClass && FoundClass->IsChildOf(UEdGraphNode::StaticClass()))
@@ -330,7 +332,7 @@ void FBlueprintSpawnNodeCommands::RegisterCommands()
 
 			InfoPtr = MakeShareable( new FEdGraphNodeSpawnInfo( FoundClass ) );
 		}
-		else if(UFunction* FoundFunction = FindObject<UFunction>(ANY_PACKAGE, *ClassName, true))
+		else if(UFunction* FoundFunction = FindFirstObject<UFunction>(*ClassName, EFindFirstObjectOptions::ExactClass, ELogVerbosity::Warning, TEXT("looking for SpawnNodes (function)")))
 		{
 			// The class name matches that of a function, so setup a spawn info that can generate function graph actions
 			InfoPtr = MakeShareable( new FFunctionNodeSpawnInfo((UFunction*)FoundFunction));
diff --git a/Engine/Source/Editor/Kismet/Private/BlueprintEditorModule.cpp b/Engine/Source/Editor/Kismet/Private/BlueprintEditorModule.cpp
--- a/Engine/Source/Editor/Kismet/Private/BlueprintEditorModule.cpp
+++ b/Engine/Source/Editor/Kismet/Private/BlueprintEditorModule.cpp
@@ -565,7 +565,7 @@ void FBlueprintEditorModule::PrepareAutoGeneratedDefaultEvents()
 			continue;
 		}
 
-		UClass* FoundTargetClass = FindObject<UClass>(ANY_PACKAGE, *TargetClassName, true);
+		UClass* FoundTargetClass = FindFirstObject<UClass>(*TargetClassName, EFindFirstObjectOptions::None, ELogVerbosity::Fatal, TEXT("looking for DefaultEventNodes"));
 		if(FoundTargetClass)
 		{
 			FString TargetEventFunction;
diff --git a/Engine/Source/Editor/Kismet/Private/BlueprintNamespaceHelper.cpp b/Engine/Source/Editor/Kismet/Private/BlueprintNamespaceHelper.cpp
--- a/Engine/Source/Editor/Kismet/Private/BlueprintNamespaceHelper.cpp
+++ b/Engine/Source/Editor/Kismet/Private/BlueprintNamespaceHelper.cpp
@@ -63,7 +63,7 @@ public:
 			return true;
 		}
 
-		FSoftObjectPath ClassPath(InBlueprint->GetClassPath());
+		FSoftObjectPath ClassPath(InBlueprint->GetClassPathName().ToString());
 		return CachedNamespaceHelper->IsImportedObject(ClassPath);
 	}
 
diff --git a/Engine/Source/Editor/Kismet/Private/BlueprintNamespaceRegistry.cpp b/Engine/Source/Editor/Kismet/Private/BlueprintNamespaceRegistry.cpp
--- a/Engine/Source/Editor/Kismet/Private/BlueprintNamespaceRegistry.cpp
+++ b/Engine/Source/Editor/Kismet/Private/BlueprintNamespaceRegistry.cpp
@@ -250,11 +250,11 @@ void FBlueprintNamespaceRegistry::FindAndRegisterAllNamespaces()
 	}
 
 	FARFilter ClassFilter;
-	ClassFilter.ClassNames.Add(UBlueprint::StaticClass()->GetFName());
-	ClassFilter.ClassNames.Add(UBlueprintGeneratedClass::StaticClass()->GetFName());
-	ClassFilter.ClassNames.Add(UUserDefinedStruct::StaticClass()->GetFName());
-	ClassFilter.ClassNames.Add(UUserDefinedEnum::StaticClass()->GetFName());
-	ClassFilter.ClassNames.Add(UBlueprintFunctionLibrary::StaticClass()->GetFName());
+	ClassFilter.ClassPaths.Add(UBlueprint::StaticClass()->GetClassPathName());
+	ClassFilter.ClassPaths.Add(UBlueprintGeneratedClass::StaticClass()->GetClassPathName());
+	ClassFilter.ClassPaths.Add(UUserDefinedStruct::StaticClass()->GetClassPathName());
+	ClassFilter.ClassPaths.Add(UUserDefinedEnum::StaticClass()->GetClassPathName());
+	ClassFilter.ClassPaths.Add(UBlueprintFunctionLibrary::StaticClass()->GetClassPathName());
 	ClassFilter.bRecursiveClasses = true;
 
 	// Register unloaded type namespace identifiers.
diff --git a/Engine/Source/Editor/Kismet/Private/CallStackViewer.cpp b/Engine/Source/Editor/Kismet/Private/CallStackViewer.cpp
--- a/Engine/Source/Editor/Kismet/Private/CallStackViewer.cpp
+++ b/Engine/Source/Editor/Kismet/Private/CallStackViewer.cpp
@@ -448,9 +448,7 @@ void SCallStackViewer::JumpToEntry(TSharedRef< FCallStackRow > Entry)
 	LastFrameNavigatedTo = Entry;
 
 	// Try to find a UClass* source:
-	bool bExactClass = false;
-	bool bAnyPackage = true;
-	UBlueprintGeneratedClass* BPGC = FindObjectFast<UBlueprintGeneratedClass>(nullptr, Entry->ScopeName, bExactClass, bAnyPackage, RF_NoFlags);
+	UBlueprintGeneratedClass* BPGC = FindFirstObject<UBlueprintGeneratedClass>(*Entry->ScopeName.ToString(), EFindFirstObjectOptions::EnsureIfAmbiguous);
 
 	if(BPGC)
 	{
diff --git a/Engine/Source/Editor/Kismet/Private/FindInBlueprintManager.cpp b/Engine/Source/Editor/Kismet/Private/FindInBlueprintManager.cpp
--- a/Engine/Source/Editor/Kismet/Private/FindInBlueprintManager.cpp
+++ b/Engine/Source/Editor/Kismet/Private/FindInBlueprintManager.cpp
@@ -1780,7 +1780,7 @@ public:
 				FAssetData AssetData = AssetRegistryModule->Get().GetAssetByObjectPath(UncachedAssets[TickCacheIndex], bIncludeOnlyOnDiskAssets);
 				if (AssetData.IsValid())
 				{
-					const bool bIsWorldAsset = AssetData.AssetClass == UWorld::StaticClass()->GetFName();
+					const bool bIsWorldAsset = AssetData.AssetClassPath == UWorld::StaticClass()->GetClassPathName();
 
 					// Construct a full package filename with path so we can query the read only status and save to disk
 					FString FinalPackageFilename = FPackageName::LongPackageNameToFilename(AssetData.PackageName.ToString());
@@ -2068,7 +2068,7 @@ void FFindInBlueprintSearchManager::OnAssetAdded(const FAssetData& InAssetData)
 {
 	const UClass* AssetClass = nullptr;
 	{
-		TWeakObjectPtr<const UClass> FoundClass = CachedAssetClasses.FindRef(InAssetData.AssetClass);
+		TWeakObjectPtr<const UClass> FoundClass = CachedAssetClasses.FindRef(InAssetData.AssetClassPath);
 		if (FoundClass.IsValid())
 		{
 			AssetClass = FoundClass.Get();
@@ -2078,7 +2078,7 @@ void FFindInBlueprintSearchManager::OnAssetAdded(const FAssetData& InAssetData)
 			AssetClass = InAssetData.GetClass();
 			if (AssetClass)
 			{
-				CachedAssetClasses.Add(InAssetData.AssetClass, TWeakObjectPtr<const UClass>(AssetClass));
+				CachedAssetClasses.Add(InAssetData.AssetClassPath, TWeakObjectPtr<const UClass>(AssetClass));
 			}
 		}
 	}
@@ -2938,9 +2938,9 @@ void FFindInBlueprintSearchManager::BuildCache()
 	FARFilter ClassFilter;
 	ClassFilter.bRecursiveClasses = true;
 
-	for (FName ClassName : FBlueprintAssetHandler::Get().GetRegisteredClassNames())
+	for (FTopLevelAssetPath ClassPathName : FBlueprintAssetHandler::Get().GetRegisteredClassNames())
 	{
-		ClassFilter.ClassNames.Add(ClassName);
+		ClassFilter.ClassPaths.Add(ClassPathName);
 	}
 
 	AssetRegistryModule->Get().GetAssets(ClassFilter, BlueprintAssets);
diff --git a/Engine/Source/Editor/Kismet/Private/FindInBlueprints.cpp b/Engine/Source/Editor/Kismet/Private/FindInBlueprints.cpp
--- a/Engine/Source/Editor/Kismet/Private/FindInBlueprints.cpp
+++ b/Engine/Source/Editor/Kismet/Private/FindInBlueprints.cpp
@@ -308,7 +308,7 @@ void FFindInBlueprintsGraphNode::FinalizeSearchData()
 {
 	if(!ClassName.IsEmpty())
 	{
-		Class = FindObject<UClass>(ANY_PACKAGE, *ClassName, true);
+		Class = UClass::TryFindTypeSlow<UClass>(ClassName, EFindFirstObjectOptions::ExactClass);
 		ClassName.Empty();
 	}
 }
@@ -371,7 +371,7 @@ void FFindInBlueprintsPin::FinalizeSearchData()
 {
 	if(!PinType.PinSubCategory.IsNone())
 	{
-		PinType.PinSubCategoryObject = FindObject<UClass>(ANY_PACKAGE, *PinType.PinSubCategory.ToString(), true);
+		PinType.PinSubCategoryObject = UClass::TryFindTypeSlow<UClass>(PinType.PinSubCategory.ToString(), EFindFirstObjectOptions::ExactClass);
 		if(!PinType.PinSubCategoryObject.IsValid())
 		{
 			PinType.PinSubCategoryObject = FindObject<UScriptStruct>(UObject::StaticClass(), *PinType.PinSubCategory.ToString());
@@ -385,7 +385,7 @@ void FFindInBlueprintsPin::FinalizeSearchData()
 
 	if(!SchemaName.IsEmpty())
 	{
-		UClass* SchemaClass = FindObject<UClass>(ANY_PACKAGE, *SchemaName, true);
+		UClass* SchemaClass = UClass::TryFindTypeSlow<UClass>(SchemaName, EFindFirstObjectOptions::ExactClass);
 		if(SchemaClass)
 		{
 			UEdGraphSchema* Schema = SchemaClass->GetDefaultObject<UEdGraphSchema>();
@@ -489,7 +489,7 @@ void FFindInBlueprintsProperty::FinalizeSearchData()
 {
 	if(!PinType.PinSubCategory.IsNone())
 	{
-		PinType.PinSubCategoryObject = FindObject<UClass>(ANY_PACKAGE, *PinType.PinSubCategory.ToString(), true);
+		PinType.PinSubCategoryObject = UClass::TryFindTypeSlow<UClass>(PinType.PinSubCategory.ToString(), EFindFirstObjectOptions::ExactClass);
 		if(!PinType.PinSubCategoryObject.IsValid())
 		{
 			PinType.PinSubCategoryObject = FindObject<UScriptStruct>(UObject::StaticClass(), *PinType.PinSubCategory.ToString());
diff --git a/Engine/Source/Editor/Kismet/Private/STimelineEditor.cpp b/Engine/Source/Editor/Kismet/Private/STimelineEditor.cpp
--- a/Engine/Source/Editor/Kismet/Private/STimelineEditor.cpp
+++ b/Engine/Source/Editor/Kismet/Private/STimelineEditor.cpp
@@ -1344,7 +1344,7 @@ void STimelineEditor::CreateNewTrack(FTTTrackBase::ETrackType Type)
 			FTTFloatTrack NewTrack;
 			NewTrack.SetTrackName(TrackName, TimelineObj);
 			// @hack for using existing curve assets.  need something better!
-			NewTrack.CurveFloat = FindObject<UCurveFloat>(ANY_PACKAGE, *TrackName.ToString() );
+			NewTrack.CurveFloat = FindFirstObject<UCurveFloat>(*TrackName.ToString(), EFindFirstObjectOptions::NativeFirst | EFindFirstObjectOptions::EnsureIfAmbiguous);
 			if (NewTrack.CurveFloat == nullptr)
 			{
 				NewTrack.CurveFloat = NewObject<UCurveFloat>(OwnerClass, NAME_None, RF_Public);
diff --git a/Engine/Source/Editor/Kismet/Public/BlueprintAssetHandler.h b/Engine/Source/Editor/Kismet/Public/BlueprintAssetHandler.h
--- a/Engine/Source/Editor/Kismet/Public/BlueprintAssetHandler.h
+++ b/Engine/Source/Editor/Kismet/Public/BlueprintAssetHandler.h
@@ -69,7 +69,7 @@ public:
 	/**
 	 * Get all the currently registered class names
 	 */
-	TArrayView<const FName> GetRegisteredClassNames() const
+	TArrayView<const FTopLevelAssetPath> GetRegisteredClassNames() const
 	{
 		return ClassNames;
 	}
@@ -82,9 +82,9 @@ public:
 	 * @param ClassName      The name of the class this handler relates to.
 	 */
 	template<typename HandlerType>
-	void RegisterHandler(FName ClassName)
+	void RegisterHandler(FTopLevelAssetPath ClassPathName)
 	{
-		RegisterHandler(ClassName, MakeUnique<HandlerType>());
+		RegisterHandler(ClassPathName, MakeUnique<HandlerType>());
 	}
 
 
@@ -92,10 +92,10 @@ public:
 	 * Register an asset for the specified class name
 	 * @note: Any assets whose class is a child of the specified class will use this handler (unless there is a more specific handler registered)
 	 *
-	 * @param ClassName      The name of the class this handler relates to.
+	 * @param ClassPathName  The path name of the class this handler relates to.
 	 * @param InHandler      The implementation of the handler to use
 	 */
-	KISMET_API void RegisterHandler(FName ClassName, TUniquePtr<IBlueprintAssetHandler>&& InHandler);
+	KISMET_API void RegisterHandler(FTopLevelAssetPath ClassPathName, TUniquePtr<IBlueprintAssetHandler>&& InHandler);
 
 
 	/**
@@ -113,7 +113,7 @@ private:
 	FBlueprintAssetHandler();
 
 	/** Unsorted array of class names, one per-handler below */
-	TArray<FName> ClassNames;
+	TArray<FTopLevelAssetPath> ClassNames;
 
 	/** Array of handlers that relate to the class names above */
 	TArray<TUniquePtr<IBlueprintAssetHandler>> Handlers;
diff --git a/Engine/Source/Editor/Kismet/Public/FindInBlueprintManager.h b/Engine/Source/Editor/Kismet/Public/FindInBlueprintManager.h
--- a/Engine/Source/Editor/Kismet/Public/FindInBlueprintManager.h
+++ b/Engine/Source/Editor/Kismet/Public/FindInBlueprintManager.h
@@ -798,7 +798,7 @@ private:
 	EFiBCacheOpType CurrentCacheOpType;
 
 	/** Mapping between a class name and its UClass instance - used for faster look up in FFindInBlueprintSearchManager::OnAssetAdded */
-	TMap<FName, TWeakObjectPtr<const UClass>> CachedAssetClasses;
+	TMap<FTopLevelAssetPath, TWeakObjectPtr<const UClass>> CachedAssetClasses;
 
 	/** The tab identifier/instance name for global find results */
 	FName GlobalFindResultsTabIDs[MAX_GLOBAL_FIND_RESULTS];
diff --git a/Engine/Source/Editor/LevelEditor/Private/LevelEditorCreateActorMenu.cpp b/Engine/Source/Editor/LevelEditor/Private/LevelEditorCreateActorMenu.cpp
--- a/Engine/Source/Editor/LevelEditor/Private/LevelEditorCreateActorMenu.cpp
+++ b/Engine/Source/Editor/LevelEditor/Private/LevelEditorCreateActorMenu.cpp
@@ -264,10 +264,7 @@ static void GetContentBrowserSelectionFactoryMenuEntries( FAssetData& TargetAsse
 		if ( TargetAssetData.GetTagValue( FBlueprintTags::NativeParentClassPath, TagValue ) && !TagValue.IsEmpty() )
 		{
 			// If the native parent class can't be placed, neither can the blueprint
-
-			UObject* Outer = nullptr;
-			ResolveName( Outer, TagValue, false, false );
-			UClass* NativeParentClass = FindObject<UClass>( ANY_PACKAGE, *TagValue );
+			UClass* NativeParentClass = UClass::TryFindTypeSlow<UClass>(FPackageName::ExportTextPathToObjectPath(TagValue));
 
 			bPlaceable = AssetSelectionUtils::IsChildBlueprintPlaceable( NativeParentClass );
 		}
diff --git a/Engine/Source/Editor/LevelEditor/Private/LevelEditorToolBar.cpp b/Engine/Source/Editor/LevelEditor/Private/LevelEditorToolBar.cpp
--- a/Engine/Source/Editor/LevelEditor/Private/LevelEditorToolBar.cpp
+++ b/Engine/Source/Editor/LevelEditor/Private/LevelEditorToolBar.cpp
@@ -1871,7 +1871,7 @@ void FLevelEditorToolBar::RegisterOpenBlueprintMenu()
 
 			// Configure filter for asset picker
 			FAssetPickerConfig Config;
-			Config.Filter.ClassNames.Add(UBlueprint::StaticClass()->GetFName());
+			Config.Filter.ClassPaths.Add(UBlueprint::StaticClass()->GetClassPathName());
 			Config.InitialAssetViewType = EAssetViewType::List;
 			Config.OnAssetSelected = FOnAssetSelected::CreateStatic(&FBlueprintMenus::OnBPSelected);
 			Config.bAllowDragging = false;
diff --git a/Engine/Source/Editor/LevelEditor/Private/SLevelViewport.cpp b/Engine/Source/Editor/LevelEditor/Private/SLevelViewport.cpp
--- a/Engine/Source/Editor/LevelEditor/Private/SLevelViewport.cpp
+++ b/Engine/Source/Editor/LevelEditor/Private/SLevelViewport.cpp
@@ -1330,34 +1330,41 @@ void SLevelViewport::BindOptionCommands( FUICommandList& OutCommandList )
 		);
 
 	IAssetRegistry & AssetRegistry = FModuleManager::LoadModuleChecked<FAssetRegistryModule>("AssetRegistry").Get();
-	TArray<FName> ClassNames;
-	TSet<FName> DerivedClassNames;
-	ClassNames.Add(ACameraActor::StaticClass()->GetFName());
-	AssetRegistry.GetDerivedClassNames(ClassNames, TSet<FName>(), DerivedClassNames);
+	TArray<FTopLevelAssetPath> ClassNames;
+	TSet<FTopLevelAssetPath> DerivedClassNames;
+	ClassNames.Add(ACameraActor::StaticClass()->GetClassPathName());
+	AssetRegistry.GetDerivedClassNames(ClassNames, TSet<FTopLevelAssetPath>(), DerivedClassNames);
 
-	for (FName Name : DerivedClassNames)
+	for (FTopLevelAssetPath ClassPathName : DerivedClassNames)
 	{
+		FString Name = ClassPathName.ToString();
 		// Ignore generated types that cannot be spawned
-		if (Name.ToString().StartsWith("SKEL_") || Name.ToString().StartsWith("REINST_"))
+		if (Name.Contains("SKEL_") || Name.Contains("REINST_"))
 		{
 			continue;
 		}
 		
 		// Remove _C from display names for blueprint classes
-		FName DisplayName = Name;
-		if (Name.ToString().EndsWith("_C"))
+		FString DisplayName = Name;
+		if (Name.EndsWith("_C"))
 		{
-			DisplayName = FName(*DisplayName.ToString().LeftChop(2));
+			DisplayName = DisplayName.LeftChop(2);
 		}
 
-		UClass* CameraClass = FindObject<UClass>(ANY_PACKAGE, *Name.ToString());
+		UClass* CameraClass = FindObject<UClass>(ClassPathName);
 		if (!CameraClass || CameraClass->HasAllClassFlags(CLASS_Abstract))
 		{
 			continue;
 		}
 		
 		// Look for existing UI Command info so one isn't created for every viewport
-		TSharedPtr<FUICommandInfo> * FoundCamera = FLevelViewportCommands::Get().CreateCameras.FindByPredicate([Name](TSharedPtr<FUICommandInfo> Camera) { return Camera->GetCommandName() == Name; });
+		FName CommandName;
+		{
+			int32 DotIndex = -1;
+			Name.FindLastChar('.', DotIndex);
+			CommandName = *Name.Mid(DotIndex + 1);
+		}
+		TSharedPtr<FUICommandInfo> * FoundCamera = FLevelViewportCommands::Get().CreateCameras.FindByPredicate([CommandName](TSharedPtr<FUICommandInfo> Camera) { return Camera->GetCommandName() == CommandName; });
 		
 		if (FoundCamera)
 		{
@@ -1369,7 +1376,7 @@ void SLevelViewport::BindOptionCommands( FUICommandList& OutCommandList )
 		else
 		{
 			// If command info isn't found, create a new one
-			TSharedRef<FUICommandInfo> NewCamera = FUICommandInfoDecl(FLevelViewportCommands::Get().AsShared(), Name, FText::FromName(DisplayName), FText::Format(LOCTEXT("SpawnCamerasTooltip", "Spawn Camera here of type {0}"), FText::FromName(Name))).UserInterfaceType(EUserInterfaceActionType::Button).DefaultChord(FInputChord());
+			TSharedRef<FUICommandInfo> NewCamera = FUICommandInfoDecl(FLevelViewportCommands::Get().AsShared(), CommandName, FText::FromString(DisplayName), FText::Format(LOCTEXT("SpawnCamerasTooltip", "Spawn Camera here of type {0}"), FText::FromString(Name))).UserInterfaceType(EUserInterfaceActionType::Button).DefaultChord(FInputChord());
 			
 			OutCommandList.MapAction(
 				NewCamera,
diff --git a/Engine/Source/Editor/MaterialEditor/Private/MaterialEditor.cpp b/Engine/Source/Editor/MaterialEditor/Private/MaterialEditor.cpp
--- a/Engine/Source/Editor/MaterialEditor/Private/MaterialEditor.cpp
+++ b/Engine/Source/Editor/MaterialEditor/Private/MaterialEditor.cpp
@@ -4759,7 +4759,7 @@ bool FMaterialEditor::CanGoToDocumentation()
 void FMaterialEditor::RenameAssetFromRegistry(const FAssetData& InAddedAssetData, const FString& InNewName)
 {
 	// Grab the asset class, it will be checked for being a material function.
-	UClass* Asset = FindObject<UClass>(ANY_PACKAGE, *InAddedAssetData.AssetClass.ToString());
+	UClass* Asset = FindObject<UClass>(InAddedAssetData.AssetClassPath);
 
 	if(Asset->IsChildOf(UMaterialFunction::StaticClass()))
 	{
diff --git a/Engine/Source/Editor/MaterialEditor/Private/MaterialEditorActions.cpp b/Engine/Source/Editor/MaterialEditor/Private/MaterialEditorActions.cpp
--- a/Engine/Source/Editor/MaterialEditor/Private/MaterialEditorActions.cpp
+++ b/Engine/Source/Editor/MaterialEditor/Private/MaterialEditorActions.cpp
@@ -115,7 +115,7 @@ void FMaterialEditorSpawnNodeCommands::RegisterCommands()
 		}
 
 		FString CommandLabel;
-		UClass* FoundClass = FindObject<UClass>(ANY_PACKAGE, *ClassName, true);
+		UClass* FoundClass = UClass::TryFindTypeSlow<UClass>(ClassName, EFindFirstObjectOptions::ExactClass);
 		TSharedPtr< FExpressionSpawnInfo > InfoPtr;
 
 		if(FoundClass && FoundClass->IsChildOf(UMaterialExpression::StaticClass()))
diff --git a/Engine/Source/Editor/MaterialEditor/Private/SMaterialPalette.cpp b/Engine/Source/Editor/MaterialEditor/Private/SMaterialPalette.cpp
--- a/Engine/Source/Editor/MaterialEditor/Private/SMaterialPalette.cpp
+++ b/Engine/Source/Editor/MaterialEditor/Private/SMaterialPalette.cpp
@@ -233,7 +233,7 @@ void SMaterialPalette::RenameAssetFromRegistry(const FAssetData& InAddedAssetDat
 void SMaterialPalette::RefreshAssetInRegistry(const FAssetData& InAddedAssetData)
 {
 	// Grab the asset class, it will be checked for being a material function.
-	UClass* Asset = FindObject<UClass>(ANY_PACKAGE, *InAddedAssetData.AssetClass.ToString());
+	UClass* Asset = FindObject<UClass>(InAddedAssetData.AssetClassPath);
 
 	if (Asset->IsChildOf(UMaterialFunction::StaticClass()))
 	{
diff --git a/Engine/Source/Editor/MergeActors/Private/MergeProxyUtils/Utils.cpp b/Engine/Source/Editor/MergeActors/Private/MergeProxyUtils/Utils.cpp
--- a/Engine/Source/Editor/MergeActors/Private/MergeProxyUtils/Utils.cpp
+++ b/Engine/Source/Editor/MergeActors/Private/MergeProxyUtils/Utils.cpp
@@ -103,7 +103,7 @@ bool GetPackageNameForMergeAction(const FString& DefaultPackageName, FString& Ou
 		SaveAssetDialogConfig.DefaultPath = DefaultPath;
 		SaveAssetDialogConfig.DefaultAssetName = DefaultName;
 		SaveAssetDialogConfig.ExistingAssetPolicy = ESaveAssetDialogExistingAssetPolicy::AllowButWarn;
-		SaveAssetDialogConfig.AssetClassNames = { UStaticMesh::StaticClass()->GetFName() };
+		SaveAssetDialogConfig.AssetClassNames = { UStaticMesh::StaticClass()->GetClassPathName() };
 
 		FContentBrowserModule& ContentBrowserModule = FModuleManager::LoadModuleChecked<FContentBrowserModule>("ContentBrowser");
 		FString SaveObjectPath = ContentBrowserModule.Get().CreateModalSaveAssetDialog(SaveAssetDialogConfig);
diff --git a/Engine/Source/Editor/MovieSceneTools/Private/K2Node_GetSequenceBinding.cpp b/Engine/Source/Editor/MovieSceneTools/Private/K2Node_GetSequenceBinding.cpp
--- a/Engine/Source/Editor/MovieSceneTools/Private/K2Node_GetSequenceBinding.cpp
+++ b/Engine/Source/Editor/MovieSceneTools/Private/K2Node_GetSequenceBinding.cpp
@@ -379,7 +379,7 @@ TSharedPtr<SGraphNode> UK2Node_GetSequenceBinding::CreateVisualWidget()
 			FContentBrowserModule& ContentBrowserModule = FModuleManager::Get().LoadModuleChecked<FContentBrowserModule>(TEXT("ContentBrowser"));
 
 			FAssetPickerConfig AssetPickerConfig;
-			AssetPickerConfig.Filter.ClassNames.Add(UMovieSceneSequence::StaticClass()->GetFName());
+			AssetPickerConfig.Filter.ClassPaths.Add(UMovieSceneSequence::StaticClass()->GetClassPathName());
 			AssetPickerConfig.bAllowNullSelection = true;
 			AssetPickerConfig.Filter.bRecursiveClasses = true;
 			AssetPickerConfig.OnAssetSelected = FOnAssetSelected::CreateSP(this, &SGraphNodeGetSequenceBinding::OnAssetSelectedFromPicker);
diff --git a/Engine/Source/Editor/MovieSceneTools/Private/MovieSceneToolHelpers.cpp b/Engine/Source/Editor/MovieSceneTools/Private/MovieSceneToolHelpers.cpp
--- a/Engine/Source/Editor/MovieSceneTools/Private/MovieSceneToolHelpers.cpp
+++ b/Engine/Source/Editor/MovieSceneTools/Private/MovieSceneToolHelpers.cpp
@@ -377,7 +377,7 @@ FString MovieSceneToolHelpers::GenerateNewShotPath(UMovieScene* SequenceMovieSce
 	FAssetRegistryModule& AssetRegistryModule = FModuleManager::LoadModuleChecked<FAssetRegistryModule>(TEXT("AssetRegistry"));
 
 	TArray<FAssetData> ObjectList;
-	AssetRegistryModule.Get().GetAssetsByClass(ULevelSequence::StaticClass()->GetFName(), ObjectList);
+	AssetRegistryModule.Get().GetAssetsByClass(ULevelSequence::StaticClass()->GetClassPathName(), ObjectList);
 
 	UObject* SequenceAsset = SequenceMovieScene->GetOuter();
 	UPackage* SequencePackage = SequenceAsset->GetOutermost();
@@ -540,7 +540,7 @@ void MovieSceneToolHelpers::GatherTakes(const UMovieSceneSection* Section, TArra
 		// Gather up all level sequence assets
 		FAssetRegistryModule& AssetRegistryModule = FModuleManager::LoadModuleChecked<FAssetRegistryModule>(TEXT("AssetRegistry"));
 		TArray<FAssetData> ObjectList;
-		AssetRegistryModule.Get().GetAssetsByClass(ULevelSequence::StaticClass()->GetFName(), ObjectList);
+		AssetRegistryModule.Get().GetAssetsByClass(ULevelSequence::StaticClass()->GetClassPathName(), ObjectList);
 
 		for (auto AssetObject : ObjectList)
 		{
@@ -596,7 +596,7 @@ bool MovieSceneToolHelpers::GetTakeNumber(const UMovieSceneSection* Section, FAs
 		// Gather up all level sequence assets
 		FAssetRegistryModule& AssetRegistryModule = FModuleManager::LoadModuleChecked<FAssetRegistryModule>(TEXT("AssetRegistry"));
 		TArray<FAssetData> ObjectList;
-		AssetRegistryModule.Get().GetAssetsByClass(ULevelSequence::StaticClass()->GetFName(), ObjectList);
+		AssetRegistryModule.Get().GetAssetsByClass(ULevelSequence::StaticClass()->GetClassPathName(), ObjectList);
 
 		for (auto AssetObject : ObjectList)
 		{
diff --git a/Engine/Source/Editor/MovieSceneTools/Private/MovieSceneTranslator.cpp b/Engine/Source/Editor/MovieSceneTools/Private/MovieSceneTranslator.cpp
--- a/Engine/Source/Editor/MovieSceneTools/Private/MovieSceneTranslator.cpp
+++ b/Engine/Source/Editor/MovieSceneTools/Private/MovieSceneTranslator.cpp
@@ -677,7 +677,7 @@ TSharedPtr<FMovieSceneImportCinematicSectionData> FMovieSceneImportData::CreateC
 
 	// Collect a full list of assets with the specified class
 	TArray<FAssetData> AssetDataArray;
-	AssetRegistryModule.Get().GetAssetsByClass(ULevelSequence::StaticClass()->GetFName(), AssetDataArray);
+	AssetRegistryModule.Get().GetAssetsByClass(ULevelSequence::StaticClass()->GetClassPathName(), AssetDataArray);
 
 	for (FAssetData AssetData : AssetDataArray)
 	{
@@ -815,7 +815,7 @@ TSharedPtr<FMovieSceneImportAudioSectionData> FMovieSceneImportData::CreateAudio
 
 	// Collect a full list of assets with the specified class
 	TArray<FAssetData> AssetDataArray;
-	AssetRegistryModule.Get().GetAssetsByClass(USoundWave::StaticClass()->GetFName(), AssetDataArray);
+	AssetRegistryModule.Get().GetAssetsByClass(USoundWave::StaticClass()->GetClassPathName(), AssetDataArray);
 
 	for (FAssetData AssetData : AssetDataArray)
 	{
diff --git a/Engine/Source/Editor/MovieSceneTools/Private/MovieSceneTranslatorEDL.cpp b/Engine/Source/Editor/MovieSceneTools/Private/MovieSceneTranslatorEDL.cpp
--- a/Engine/Source/Editor/MovieSceneTools/Private/MovieSceneTranslatorEDL.cpp
+++ b/Engine/Source/Editor/MovieSceneTools/Private/MovieSceneTranslatorEDL.cpp
@@ -404,7 +404,7 @@ bool MovieSceneTranslatorEDL::ImportEDL(UMovieScene* InMovieScene, FFrameRate In
 
 				// Collect a full list of assets with the specified class
 				TArray<FAssetData> AssetDataArray;
-				AssetRegistryModule.Get().GetAssetsByClass(ULevelSequence::StaticClass()->GetFName(), AssetDataArray);
+				AssetRegistryModule.Get().GetAssetsByClass(ULevelSequence::StaticClass()->GetClassPathName(), AssetDataArray);
 
 				for (FAssetData AssetData : AssetDataArray)
 				{
diff --git a/Engine/Source/Editor/MovieSceneTools/Private/TrackEditors/AudioTrackEditor.cpp b/Engine/Source/Editor/MovieSceneTools/Private/TrackEditors/AudioTrackEditor.cpp
--- a/Engine/Source/Editor/MovieSceneTools/Private/TrackEditors/AudioTrackEditor.cpp
+++ b/Engine/Source/Editor/MovieSceneTools/Private/TrackEditors/AudioTrackEditor.cpp
@@ -1293,10 +1293,10 @@ TSharedRef<SWidget> FAudioTrackEditor::BuildAudioSubMenu(FOnAssetSelected OnAsse
 	UMovieSceneSequence* Sequence = GetSequencer() ? GetSequencer()->GetFocusedMovieSceneSequence() : nullptr;
 
 	FAssetRegistryModule& AssetRegistryModule = FModuleManager::LoadModuleChecked<FAssetRegistryModule>(TEXT("AssetRegistry"));
-	TArray<FName> ClassNames;
-	ClassNames.Add(USoundBase::StaticClass()->GetFName());
-	TSet<FName> DerivedClassNames;
-	AssetRegistryModule.Get().GetDerivedClassNames(ClassNames, TSet<FName>(), DerivedClassNames);
+	TArray<FTopLevelAssetPath> ClassNames;
+	ClassNames.Add(USoundBase::StaticClass()->GetClassPathName());
+	TSet<FTopLevelAssetPath> DerivedClassNames;
+	AssetRegistryModule.Get().GetDerivedClassNames(ClassNames, TSet<FTopLevelAssetPath>(), DerivedClassNames);
 
 	FMenuBuilder MenuBuilder(true, nullptr);
 
@@ -1306,9 +1306,9 @@ TSharedRef<SWidget> FAudioTrackEditor::BuildAudioSubMenu(FOnAssetSelected OnAsse
 		AssetPickerConfig.OnAssetEnterPressed = OnAssetEnterPressed;
 		AssetPickerConfig.bAllowNullSelection = false;
 		AssetPickerConfig.InitialAssetViewType = EAssetViewType::List;
-		for (auto ClassName : DerivedClassNames)
+		for (FTopLevelAssetPath ClassName : DerivedClassNames)
 		{
-			AssetPickerConfig.Filter.ClassNames.Add(ClassName);
+			AssetPickerConfig.Filter.ClassPaths.Add(ClassName);
 		}
 		AssetPickerConfig.SaveSettingsName = TEXT("SequencerAssetPicker");
 		AssetPickerConfig.AdditionalReferencingAssets.Add(FAssetData(Sequence));
diff --git a/Engine/Source/Editor/MovieSceneTools/Private/TrackEditors/CameraAnimTrackEditor.cpp b/Engine/Source/Editor/MovieSceneTools/Private/TrackEditors/CameraAnimTrackEditor.cpp
--- a/Engine/Source/Editor/MovieSceneTools/Private/TrackEditors/CameraAnimTrackEditor.cpp
+++ b/Engine/Source/Editor/MovieSceneTools/Private/TrackEditors/CameraAnimTrackEditor.cpp
@@ -118,7 +118,7 @@ void FCameraAnimTrackEditor::BuildObjectBindingTrackMenu(FMenuBuilder& MenuBuild
 
 		// Collect a full list of assets with the specified class
 		TArray<FAssetData> AssetDataList;
-		AssetRegistryModule.Get().GetAssetsByClass(UCameraAnim::StaticClass()->GetFName(), AssetDataList);
+		AssetRegistryModule.Get().GetAssetsByClass(UCameraAnim::StaticClass()->GetClassPathName(), AssetDataList);
 
 		MenuBuilder.AddSubMenu(
 			LOCTEXT("AddCameraAnim", "Camera Anim"), NSLOCTEXT("Sequencer", "AddCameraAnimTooltip", "Adds an additive camera animation track."),
@@ -150,7 +150,7 @@ void FCameraAnimTrackEditor::AddCameraAnimSubMenu(FMenuBuilder& MenuBuilder, TAr
 		AssetPickerConfig.OnAssetEnterPressed = FOnAssetEnterPressed::CreateRaw(this, &FCameraAnimTrackEditor::OnCameraAnimAssetEnterPressed, ObjectBinding);
 		AssetPickerConfig.bAllowNullSelection = false;
 		AssetPickerConfig.InitialAssetViewType = EAssetViewType::List;
-		AssetPickerConfig.Filter.ClassNames.Add(UCameraAnim::StaticClass()->GetFName());
+		AssetPickerConfig.Filter.ClassPaths.Add(UCameraAnim::StaticClass()->GetClassPathName());
 		AssetPickerConfig.SaveSettingsName = TEXT("SequencerAssetPicker");
 		AssetPickerConfig.AdditionalReferencingAssets.Add(FAssetData(Sequence));
 	}
diff --git a/Engine/Source/Editor/MovieSceneTools/Private/TrackEditors/CameraShakeSourceShakeTrackEditor.cpp b/Engine/Source/Editor/MovieSceneTools/Private/TrackEditors/CameraShakeSourceShakeTrackEditor.cpp
--- a/Engine/Source/Editor/MovieSceneTools/Private/TrackEditors/CameraShakeSourceShakeTrackEditor.cpp
+++ b/Engine/Source/Editor/MovieSceneTools/Private/TrackEditors/CameraShakeSourceShakeTrackEditor.cpp
@@ -431,25 +431,24 @@ void FCameraShakeSourceShakeTrackEditor::AddOtherCameraShakeBrowserSubMenu(FMenu
 		AssetPickerConfig.OnShouldFilterAsset = FOnShouldFilterAsset::CreateSP(this, &FCameraShakeSourceShakeTrackEditor::OnShouldFilterCameraShake);
 		AssetPickerConfig.bAllowNullSelection = false;
 		AssetPickerConfig.InitialAssetViewType = EAssetViewType::List;
-		AssetPickerConfig.Filter.ClassNames.Add(UBlueprint::StaticClass()->GetFName());
+		AssetPickerConfig.Filter.ClassPaths.Add(UBlueprint::StaticClass()->GetClassPathName());
 		AssetPickerConfig.SaveSettingsName = TEXT("SequencerAssetPicker");
 		AssetPickerConfig.AdditionalReferencingAssets.Add(FAssetData(Sequence));
 
 		IAssetRegistry & AssetRegistry = FModuleManager::LoadModuleChecked<FAssetRegistryModule>("AssetRegistry").Get();
-		TArray<FName> ClassNames;
-		TSet<FName> DerivedClassNames;
-		ClassNames.Add(UCameraShakeBase::StaticClass()->GetFName());
-		AssetRegistry.GetDerivedClassNames(ClassNames, TSet<FName>(), DerivedClassNames);
+		TArray<FTopLevelAssetPath> ClassNames;
+		TSet<FTopLevelAssetPath> DerivedClassNames;
+		ClassNames.Add(UCameraShakeBase::StaticClass()->GetClassPathName());
+		AssetRegistry.GetDerivedClassNames(ClassNames, TSet<FTopLevelAssetPath>(), DerivedClassNames);
 						
 		AssetPickerConfig.OnShouldFilterAsset = FOnShouldFilterAsset::CreateLambda([DerivedClassNames](const FAssetData& AssetData)
 		{
 			const FString ParentClassFromData = AssetData.GetTagValueRef<FString>(FBlueprintTags::ParentClassPath);
 			if (!ParentClassFromData.IsEmpty())
 			{
-				const FString ClassObjectPath = FPackageName::ExportTextPathToObjectPath(ParentClassFromData);
-				const FName ClassName = FName(*FPackageName::ObjectPathToObjectName(ClassObjectPath));
+				const FTopLevelAssetPath ClassObjectPath(FPackageName::ExportTextPathToObjectPath(ParentClassFromData));
 
-				if (DerivedClassNames.Contains(ClassName))
+				if (DerivedClassNames.Contains(ClassObjectPath))
 				{
 					return false;
 				}
diff --git a/Engine/Source/Editor/MovieSceneTools/Private/TrackEditors/CameraShakeTrackEditor.cpp b/Engine/Source/Editor/MovieSceneTools/Private/TrackEditors/CameraShakeTrackEditor.cpp
--- a/Engine/Source/Editor/MovieSceneTools/Private/TrackEditors/CameraShakeTrackEditor.cpp
+++ b/Engine/Source/Editor/MovieSceneTools/Private/TrackEditors/CameraShakeTrackEditor.cpp
@@ -132,25 +132,23 @@ void FCameraShakeTrackEditor::AddCameraShakeSubMenu(FMenuBuilder& MenuBuilder, T
 		AssetPickerConfig.OnAssetEnterPressed = FOnAssetEnterPressed::CreateRaw(this, &FCameraShakeTrackEditor::OnCameraShakeAssetEnterPressed, ObjectBindings);
 		AssetPickerConfig.bAllowNullSelection = false;
 		AssetPickerConfig.InitialAssetViewType = EAssetViewType::List;
-		AssetPickerConfig.Filter.ClassNames.Add(UBlueprint::StaticClass()->GetFName());
+		AssetPickerConfig.Filter.ClassPaths.Add(UBlueprint::StaticClass()->GetClassPathName());
 		AssetPickerConfig.SaveSettingsName = TEXT("SequencerAssetPicker");
 		AssetPickerConfig.AdditionalReferencingAssets.Add(FAssetData(Sequence));
 
 		IAssetRegistry & AssetRegistry = FModuleManager::LoadModuleChecked<FAssetRegistryModule>("AssetRegistry").Get();
-		TArray<FName> ClassNames;
-		TSet<FName> DerivedClassNames;
-		ClassNames.Add(UCameraShakeBase::StaticClass()->GetFName());
-		AssetRegistry.GetDerivedClassNames(ClassNames, TSet<FName>(), DerivedClassNames);
+		TArray<FTopLevelAssetPath> ClassNames;
+		TSet<FTopLevelAssetPath> DerivedClassNames;
+		ClassNames.Add(UCameraShakeBase::StaticClass()->GetClassPathName());
+		AssetRegistry.GetDerivedClassNames(ClassNames, TSet<FTopLevelAssetPath>(), DerivedClassNames);
 						
 		AssetPickerConfig.OnShouldFilterAsset = FOnShouldFilterAsset::CreateLambda([DerivedClassNames](const FAssetData& AssetData)
 		{
 			const FString ParentClassFromData = AssetData.GetTagValueRef<FString>(FBlueprintTags::ParentClassPath);
 			if (!ParentClassFromData.IsEmpty())
 			{
-				const FString ClassObjectPath = FPackageName::ExportTextPathToObjectPath(ParentClassFromData);
-				const FName ClassName = FName(*FPackageName::ObjectPathToObjectName(ClassObjectPath));
-
-				if (DerivedClassNames.Contains(ClassName))
+				const FTopLevelAssetPath ClassObjectPath(FPackageName::ExportTextPathToObjectPath(ParentClassFromData));
+				if (DerivedClassNames.Contains(ClassObjectPath))
 				{
 					return false;
 				}
