Adding BoldItalic and Italic font style to FStyleFonts, use cases in AnimCurve viewer and Class Viewer respectively

diff --git a/Engine/Source/Editor/ClassViewer/Private/SClassViewer.cpp b/Engine/Source/Editor/ClassViewer/Private/SClassViewer.cpp
--- a/Engine/Source/Editor/ClassViewer/Private/SClassViewer.cpp
+++ b/Engine/Source/Editor/ClassViewer/Private/SClassViewer.cpp
@@ -108,7 +108,7 @@ public:
 	FClassHierarchy();
 	~FClassHierarchy();
 
-	/** Populates the class hierarchy tree, pulling all the loaded and unloaded classes into a master tree. */
+	/** Populates the class hierarchy tree, pulling in all the loaded and unloaded classes. */
 	void PopulateClassHierarchy();
 	void PopulateClassHierarchy(const FAssetData& InAssetData) { PopulateClassHierarchy(); }
 
@@ -766,7 +766,7 @@ public:
 		, _bIsInClassViewer( true )
 		, _bDynamicClassLoading( true )
 		, _HighlightText()
-		, _TextColor(FLinearColor(1.0f, 1.0f, 1.0f, 1.0f))
+		, _Font(FAppStyle::Get().GetFontStyle("NormalFont"))
 		{}
 
 		/** The classname this item contains. */
@@ -779,8 +779,8 @@ public:
 		SLATE_ARGUMENT( bool, bDynamicClassLoading )
 		/** The text this item should highlight, if any. */
 		SLATE_ARGUMENT( FText, HighlightText )
-		/** The color text this item will use. */
-		SLATE_ARGUMENT( FSlateColor, TextColor )
+		/** The font this item will use. */
+		SLATE_ARGUMENT( FSlateFontInfo, Font )
 		/** The node this item is associated with. */
 		SLATE_ARGUMENT( TSharedPtr<FClassViewerNode>, AssociatedNode)
 		/** the delegate for handling double clicks outside of the SClassItem */
@@ -868,8 +868,9 @@ public:
 				[
 					SNew( STextBlock )
 						.Text( FText::FromString(*ClassName.Get()) )
+						.Font(InArgs._Font)
 						.HighlightText(InArgs._HighlightText)
-						.ColorAndOpacity( this, &SClassItem::GetTextColor)
+						.ColorAndOpacity(FSlateColor::UseForeground())
 						.ToolTip(Local::GetToolTip(AssociatedNode))
 						.IsEnabled(!bIsRestricted)
 				]
@@ -886,8 +887,6 @@ public:
 						.OnGetMenuContent(this, &SClassItem::GenerateDropDown)
 				]
 		];
-		
-		TextColor = InArgs._TextColor;
 
 		UE_LOG(LogEditorClassViewer, VeryVerbose, TEXT("CLASS [%s]"), **ClassName);
 
@@ -963,21 +962,6 @@ private:
 		return SNullWidget::NullWidget;
 	}
 
-	/** Returns the text color for the item based on if it is selected or not. */
-	FSlateColor GetTextColor() const
-	{
-		const TSharedPtr< ITypedTableView< TSharedPtr<FString> > > OwnerWidget = OwnerTablePtr.Pin();
-		const TSharedPtr<FString>* MyItem = OwnerWidget->Private_ItemFromWidget( this );
-		const bool bIsSelected = OwnerWidget->Private_IsItemSelected( *MyItem );
-
-		if(bIsSelected)
-		{
-			return FSlateColor::UseForeground();
-		}
-
-		return TextColor;
-	}
-
 private:
 
 	/** The class name for which this item is associated with. */
@@ -992,9 +976,6 @@ private:
 	/** true if dynamic class loading is permitted. */
 	bool bDynamicClassLoading;
 
-	/** The text color for this item. */
-	FSlateColor TextColor;
-
 	/** The Class Viewer Node this item is associated with. */
 	TSharedPtr< FClassViewerNode > AssociatedNode;
 
@@ -1864,7 +1845,7 @@ TSharedRef< ITableRow > SClassViewer::OnGenerateRowForClassViewer( TSharedPtr<FC
 		.ClassName(ClassNameDisplay)
 		.bIsPlaceable(Item->IsClassPlaceable())
 		.HighlightText(SearchBoxTextForHighlight)
-		.TextColor(Item->IsClassPlaceable()? FLinearColor(0.2f, 0.4f, 0.6f, AlphaValue) : FLinearColor(1.0f, 1.0f, 1.0f, AlphaValue))
+		.Font(Item->IsClassPlaceable()? FCoreStyle::Get().GetFontStyle("NormalFontItalic") : FCoreStyle::Get().GetFontStyle("NormalFont"))
 		.AssociatedNode(Item)
 		.bIsInClassViewer( InitOptions.Mode == EClassViewerMode::ClassBrowsing )
 		.bDynamicClassLoading( bEnableClassDynamicLoading )
diff --git a/Engine/Source/Editor/Persona/Private/SAnimCurveViewer.cpp b/Engine/Source/Editor/Persona/Private/SAnimCurveViewer.cpp
--- a/Engine/Source/Editor/Persona/Private/SAnimCurveViewer.cpp
+++ b/Engine/Source/Editor/Persona/Private/SAnimCurveViewer.cpp
@@ -75,7 +75,7 @@ TSharedRef< SWidget > SAnimCurveListRow::GenerateWidgetForColumn( const FName& C
 				[
 					SAssignNew(Item->EditableText, SInlineEditableTextBlock)
 					.OnTextCommitted(AnimCurveViewer.Get(), &SAnimCurveViewer::OnNameCommitted, Item)
-					.ColorAndOpacity(this, &SAnimCurveListRow::GetItemTextColor)
+					.Font(this, &SAnimCurveListRow::GetItemFont)
 					.IsSelected(this, &SAnimCurveListRow::IsSelected)
 					.Text(this, &SAnimCurveListRow::GetItemName)
 					.HighlightText(this, &SAnimCurveListRow::GetFilterText)
@@ -380,14 +380,8 @@ bool SAnimCurveListRow::GetActiveWeight(float& OutWeight) const
 }
 
 
-FSlateColor SAnimCurveListRow::GetItemTextColor() const
+FSlateFontInfo SAnimCurveListRow::GetItemFont() const
 {
-	// If row is selected, show text as black to make it easier to read
-	if (IsSelected())
-	{
-		return FLinearColor(0, 0, 0);
-	}
-
 	// If not selected, show bright if active
 	bool bItemActive = true;
 	if (Item->bAutoFillData)
@@ -398,7 +392,7 @@ FSlateColor SAnimCurveListRow::GetItemTextColor() const
 		bItemActive = (Weight != 0.f);
 	}
 
-	return bItemActive ? FLinearColor(1, 1, 1) : FLinearColor(0.5, 0.5, 0.5);
+	return bItemActive ? FAppStyle::Get().GetFontStyle("NormalFontBoldItalic") : FAppStyle::Get().GetFontStyle("NormalFont");
 }
 
 float SAnimCurveListRow::GetWeight() const 
diff --git a/Engine/Source/Editor/Persona/Private/SAnimCurveViewer.h b/Engine/Source/Editor/Persona/Private/SAnimCurveViewer.h
--- a/Engine/Source/Editor/Persona/Private/SAnimCurveViewer.h
+++ b/Engine/Source/Editor/Persona/Private/SAnimCurveViewer.h
@@ -112,8 +112,8 @@ private:
 	FText GetItemName() const;
 	/** Get text we are filtering for */
 	FText GetFilterText() const;
-	/** Return color for text of item */
-	FSlateColor GetItemTextColor() const;
+	/** Return font for text of item */
+	FSlateFontInfo GetItemFont() const;
 
 	/** Get current active weight. Returns false if not currently active */
 	bool GetActiveWeight(float& OutWeight) const;
diff --git a/Engine/Source/Runtime/SlateCore/Private/Styling/StarshipCoreStyle.cpp b/Engine/Source/Runtime/SlateCore/Private/Styling/StarshipCoreStyle.cpp
--- a/Engine/Source/Runtime/SlateCore/Private/Styling/StarshipCoreStyle.cpp
+++ b/Engine/Source/Runtime/SlateCore/Private/Styling/StarshipCoreStyle.cpp
@@ -34,7 +34,9 @@ TUniquePtr<FStyleFonts> FStyleFonts::Instance = nullptr;
 
 FStyleFonts::FStyleFonts()
 	: Normal(FONT(10, "Regular"))
+	, NormalItalic(FONT(10, "Italic"))
 	, NormalBold(FONT(10, "Bold"))
+	, NormalBoldItalic(FONT(10, "BoldItalic"))
 	, Small(FONT(8,  "Regular"))
 	, SmallBold(FONT(8,  "Bold"))
 	, Large(FONT(14, "Regular"))
@@ -1075,6 +1077,8 @@ void FStarshipCoreStyle::SetupTextStyles(TSharedRef<FStyle>& Style)
 	Style->Set("NormalFontBold", StyleFonts.NormalBold);
 	Style->Set("SmallFontBold", StyleFonts.SmallBold);
 	Style->Set("Font.Large.Bold", StyleFonts.Small);
+	Style->Set("NormalFontItalic", StyleFonts.NormalItalic);
+	Style->Set("NormalFontBoldItalic", StyleFonts.NormalBoldItalic);
 
 	Style->Set("HeadingMedium", StyleFonts.HeadingMedium);
 	Style->Set("HeadingSmall", StyleFonts.HeadingSmall);
diff --git a/Engine/Source/Runtime/SlateCore/Public/Styling/StarshipCoreStyle.h b/Engine/Source/Runtime/SlateCore/Public/Styling/StarshipCoreStyle.h
--- a/Engine/Source/Runtime/SlateCore/Public/Styling/StarshipCoreStyle.h
+++ b/Engine/Source/Runtime/SlateCore/Public/Styling/StarshipCoreStyle.h
@@ -101,7 +101,9 @@ struct SLATECORE_API FStyleFonts
 	}
 
 	const FSlateFontInfo Normal;
+	const FSlateFontInfo NormalItalic;
 	const FSlateFontInfo NormalBold;
+	const FSlateFontInfo NormalBoldItalic;
 	const FSlateFontInfo Small;
 	const FSlateFontInfo SmallBold;
 	const FSlateFontInfo Large;
