Fix class viewer filter asset reference check on a cooked BP (which only has the BPGC asset)
Log a warning if the BP class cannot be found (which can happen when an uncooked BP has been moved without being recompiled)

diff --git a/Engine/Source/Editor/ClassViewer/Private/ClassViewerFilter.cpp b/Engine/Source/Editor/ClassViewer/Private/ClassViewerFilter.cpp
--- a/Engine/Source/Editor/ClassViewer/Private/ClassViewerFilter.cpp
+++ b/Engine/Source/Editor/ClassViewer/Private/ClassViewerFilter.cpp
@@ -649,8 +649,22 @@ PRAGMA_ENABLE_DEPRECATION_WARNINGS
 	bool bPassesAssetReferenceFilter = true;
 	if (AssetReferenceFilter.IsValid() && bIsBlueprintBase)
 	{
-		FName BlueprintPath = FName(*GeneratedClassPathString.LeftChop(2)); // Chop off _C
-		bPassesAssetReferenceFilter = AssetReferenceFilter->PassesFilter(AssetRegistry.GetAssetByObjectPath(BlueprintPath));
+		FName BlueprintAssetPath = FName(*GeneratedClassPathString);
+		FAssetData BlueprintAssetData = AssetRegistry.GetAssetByObjectPath(BlueprintAssetPath);
+		if (!BlueprintAssetData.IsValid())
+		{
+			BlueprintAssetPath = FName(*GeneratedClassPathString.LeftChop(2)); // Chop off _C
+			BlueprintAssetData = AssetRegistry.GetAssetByObjectPath(BlueprintAssetPath);
+		}
+
+		if (BlueprintAssetData.IsValid())
+		{
+			bPassesAssetReferenceFilter = AssetReferenceFilter->PassesFilter(BlueprintAssetData);
+		}
+		else
+		{
+			UE_LOG(LogEditorClassViewer, Warning, TEXT("Blueprint class cannot be found: %s"), *GeneratedClassPathString);
+		}
 	}
 
 	bool bPassesFilter = bPassesPlaceableFilter && bPassesBlueprintBaseFilter && bPassesDeveloperFilter 
@@ -658,4 +672,4 @@ PRAGMA_ENABLE_DEPRECATION_WARNINGS
 		&& (!bCheckTextFilter || bPassesTextFilter) && bPassesAssetReferenceFilter;
 
 	return bPassesFilter;
-}
\ No newline at end of file
+}
diff --git a/Engine/Source/Editor/ClassViewer/Private/SClassViewer.cpp b/Engine/Source/Editor/ClassViewer/Private/SClassViewer.cpp
--- a/Engine/Source/Editor/ClassViewer/Private/SClassViewer.cpp
+++ b/Engine/Source/Editor/ClassViewer/Private/SClassViewer.cpp
@@ -78,8 +78,6 @@
 
 #define LOCTEXT_NAMESPACE "SClassViewer"
 
-DEFINE_LOG_CATEGORY_STATIC(LogEditorClassViewer, Log, All);
-
 //////////////////////////////////////////////////////////////
 
 struct FClassViewerNodeNameLess
diff --git a/Engine/Source/Editor/ClassViewer/Public/ClassViewerModule.h b/Engine/Source/Editor/ClassViewer/Public/ClassViewerModule.h
--- a/Engine/Source/Editor/ClassViewer/Public/ClassViewerModule.h
+++ b/Engine/Source/Editor/ClassViewer/Public/ClassViewerModule.h
@@ -11,6 +11,8 @@
 class IPropertyHandle;
 class SWidget;
 
+DEFINE_LOG_CATEGORY_STATIC(LogEditorClassViewer, Log, All);
+
 /** Delegate used with the Class Viewer in 'class picking' mode.  You'll bind a delegate when the
     class viewer widget is created, which will be fired off when a class is selected in the list */
 DECLARE_DELEGATE_OneParam( FOnClassPicked, UClass* );
