Fix FClassHierarchy not taking CoreRedirectors into consideration when generating the class hierarchy.

diff --git a/Engine/Source/Editor/ClassViewer/Private/SClassViewer.cpp b/Engine/Source/Editor/ClassViewer/Private/SClassViewer.cpp
--- a/Engine/Source/Editor/ClassViewer/Private/SClassViewer.cpp
+++ b/Engine/Source/Editor/ClassViewer/Private/SClassViewer.cpp
@@ -78,6 +78,7 @@
 #include "Trace/Detail/Channel.h"
 #include "Types/SlateStructs.h"
 #include "UObject/Class.h"
+#include "UObject/CoreRedirects.h"
 #include "UObject/Object.h"
 #include "UObject/ObjectMacros.h"
 #include "UObject/Package.h"
@@ -229,6 +230,9 @@ private:
 
 	/** Callback registered to the Asset Registry to be notified when an asset is removed. */
 	void RemoveAsset(const FAssetData& InRemovedAssetData);
+
+	/** Apply CoreRedirects on TopLevelAssetPath. */
+	void FixupClassCoreRedirects(FTopLevelAssetPath& InOutClassPath);
 private:
 	/** The "Object" class node that is used as a rooting point for the Class Viewer. */
 	TSharedPtr< FClassViewerNode > ObjectClassRoot;
@@ -1173,25 +1177,32 @@ void FClassHierarchy::UpdateClassInNode(FTopLevelAssetPath InGeneratedClassPath,
 
 bool FClassHierarchy::FindAndRemoveNodeByClassPath(const TSharedPtr< FClassViewerNode >& InRootNode, FTopLevelAssetPath InClassPath)
 {
-	bool bReturnValue = false;
+	FixupClassCoreRedirects(InClassPath);
 
-	// Search the children recursively, one of them might have the parent.
-	for(int32 ChildClassIndex = 0; ChildClassIndex < InRootNode->GetChildrenList().Num(); ChildClassIndex++)
-	{
-		if(InRootNode->GetChildrenList()[ChildClassIndex]->ClassPath == InClassPath)						   
+	TFunction<bool(const TSharedPtr<FClassViewerNode>&)> FindAndRemoveNodeRecursive;
+	FindAndRemoveNodeRecursive = [&InClassPath, &FindAndRemoveNodeRecursive](const TSharedPtr<FClassViewerNode>& InRootNode)
 		{
-			InRootNode->GetChildrenList().RemoveAt(ChildClassIndex);
-			return true;
-		}
-		// Check the child, then check the return to see if it is valid. If it is valid, end the recursion.
-		bReturnValue = FindAndRemoveNodeByClassPath(InRootNode->GetChildrenList()[ChildClassIndex], InClassPath);
+			// Search the children recursively, one of them might have the parent.
+			TArray<TSharedPtr<FClassViewerNode>>& ChildrenList = InRootNode->GetChildrenList();
+			for (int32 ChildClassIndex = 0; ChildClassIndex < ChildrenList.Num(); ChildClassIndex++)
+			{
+				if (ChildrenList[ChildClassIndex]->ClassPath == InClassPath)
+				{
+					ChildrenList.RemoveAt(ChildClassIndex);
+					return true;
+				}
 
-		if(bReturnValue)
-		{
-			break;
-		}
-	}
-	return bReturnValue;
+				// Check the child, then check the return to see if it is valid. If it is valid, end the recursion.
+				if (FindAndRemoveNodeRecursive(ChildrenList[ChildClassIndex]))
+				{
+					return true;
+				}
+			}
+
+			return false;
+		};
+
+	return FindAndRemoveNodeRecursive(InRootNode);
 }
 
 void FClassHierarchy::RemoveAsset(const FAssetData& InRemovedAssetData)
@@ -1293,6 +1304,28 @@ void FClassHierarchy::SetClassFields(TSharedPtr<FClassViewerNode>& InOutClassNod
 	InOutClassNode->Class = &Class;
 }
 
+
+void FClassHierarchy::FixupClassCoreRedirects(FTopLevelAssetPath& InOutClassPath)
+{
+	const FCoreRedirectObjectName OldName = FCoreRedirectObjectName(InOutClassPath);
+	const FCoreRedirectObjectName NewName = FCoreRedirects::GetRedirectedName(ECoreRedirectFlags::Type_Class, OldName);
+	
+	if (OldName != NewName)
+	{
+		const FString OldClassPath = InOutClassPath.ToString();
+
+		// Only do the fixup if the old class name isn't in memory.
+		const UClass* FoundOldClass = UClass::TryFindTypeSlow<UClass>(OldClassPath);
+		const FTopLevelAssetPath NewClassPath(NewName.ToString());
+
+		if (!FoundOldClass || FoundOldClass->GetClassPathName() == NewClassPath)
+		{
+			InOutClassPath = NewClassPath;
+		}
+	}
+}
+
+
 void FClassHierarchy::CreateOrUpdateUnloadedClassNode(TSharedPtr<FClassViewerNode>& InOutClassViewerNode, const FAssetData& InAssetData, FTopLevelAssetPath InClassPath)
 {
 	if (!InOutClassViewerNode)
@@ -1321,7 +1354,9 @@ void FClassHierarchy::CreateOrUpdateUnloadedClassNode(TSharedPtr<FClassViewerNod
 		FString ParentClassPathString;
 		if (InAssetData.GetTagValue(FBlueprintTags::ParentClassPath, ParentClassPathString))
 		{
-			InOutClassViewerNode->ParentClassPath = FTopLevelAssetPath(*FPackageName::ExportTextPathToObjectPath(ParentClassPathString));
+			FTopLevelAssetPath ParentClassPath(FTopLevelAssetPath(*FPackageName::ExportTextPathToObjectPath(ParentClassPathString)));
+			FixupClassCoreRedirects(ParentClassPath);
+			InOutClassViewerNode->ParentClassPath = ParentClassPath;
 		}
 	}
 
@@ -1367,9 +1402,10 @@ void FClassHierarchy::PopulateClassHierarchy()
 
 		for (const FAssetData& AssetData : Assets)
 		{
-			const FTopLevelAssetPath ClassPath = FEditorClassUtils::GetClassPathNameFromAssetTag(AssetData);
+			FTopLevelAssetPath ClassPath = FEditorClassUtils::GetClassPathNameFromAssetTag(AssetData);
 			if (!ClassPath.IsNull())
 			{
+				FixupClassCoreRedirects(ClassPath);
 				TSharedPtr<FClassViewerNode>& Node = ClassPathToNode.FindOrAdd(ClassPath);
 				CreateOrUpdateUnloadedClassNode(Node, AssetData, ClassPath);
 			}
@@ -1386,6 +1422,7 @@ void FClassHierarchy::PopulateClassHierarchy()
 		for (const FAssetData& AssetData : Assets)
 		{
 			FTopLevelAssetPath ClassPathNameFromAssetPath = AssetData.GetSoftObjectPath().GetAssetPath();
+			FixupClassCoreRedirects(ClassPathNameFromAssetPath);
 			TSharedPtr<FClassViewerNode>& Node = ClassPathToNode.FindOrAdd(ClassPathNameFromAssetPath);
 			CreateOrUpdateUnloadedClassNode(Node, AssetData, ClassPathNameFromAssetPath);
 		}
