Optimize editor class pickers by skipping expensive filter checks for classes that fail the fast ones
Stop filter and blueprint namespace helpers from doing redundant object lookups
Fix unnecessary string conversions with FTopLevelAssetPath

diff --git a/Engine/Source/Editor/ClassViewer/Private/ClassViewerFilter.cpp b/Engine/Source/Editor/ClassViewer/Private/ClassViewerFilter.cpp
--- a/Engine/Source/Editor/ClassViewer/Private/ClassViewerFilter.cpp
+++ b/Engine/Source/Editor/ClassViewer/Private/ClassViewerFilter.cpp
@@ -447,8 +447,15 @@ bool FClassViewerFilter::IsClassAllowed(const FClassViewerInitializationOptions&
 		return false;
 	}
 
-	const bool bPassesBlueprintBaseFilter = !InInitOptions.bIsBlueprintBaseOnly || CanCreateBlueprintOfClass(const_cast<UClass*>(InClass));
-	const bool bPassesEditorClassFilter = !InInitOptions.bEditorClassesOnly || IsEditorOnlyObject(InClass);
+	if (InInitOptions.bIsBlueprintBaseOnly && !CanCreateBlueprintOfClass(const_cast<UClass*>(InClass)))
+	{
+		return false;
+	}
+
+	if (InInitOptions.bEditorClassesOnly && !IsEditorOnlyObject(InClass))
+	{
+		return false;
+	}
 
 	// Determine if we allow any developer folder classes, if so determine if this class is in one of the allowed developer folders.
 	static const FString DeveloperPathWithSlash = FPackageName::FilenameToLongPackageName(FPaths::GameDevelopersDir());
@@ -457,117 +464,122 @@ bool FClassViewerFilter::IsClassAllowed(const FClassViewerInitializationOptions&
 
 	const UClassViewerSettings* ClassViewerSettings = GetDefault<UClassViewerSettings>();
 
-	bool bPassesDeveloperFilter = true;
 	EClassViewerDeveloperType AllowedDeveloperType = ClassViewerSettings->DeveloperFolderType;
 	if (AllowedDeveloperType == EClassViewerDeveloperType::CVDT_None)
 	{
-		bPassesDeveloperFilter = !GeneratedClassPathString.StartsWith(DeveloperPathWithSlash);
+		if (GeneratedClassPathString.StartsWith(DeveloperPathWithSlash))
+		{
+			return false;
+		}
 	}
 	else if (AllowedDeveloperType == EClassViewerDeveloperType::CVDT_CurrentUser)
 	{
-		if (GeneratedClassPathString.StartsWith(DeveloperPathWithSlash))
+		if (GeneratedClassPathString.StartsWith(DeveloperPathWithSlash) && !GeneratedClassPathString.StartsWith(UserDeveloperPathWithSlash))
 		{
-			bPassesDeveloperFilter = GeneratedClassPathString.StartsWith(UserDeveloperPathWithSlash);
+			return false;
 		}
 	}
 	
 	// The INI files declare classes and folders that are considered internal only. Does this class match any of those patterns?
 	// INI path: /Script/ClassViewer.ClassViewerProjectSettings
-	bool bPassesInternalFilter = true;
 	if (!ClassViewerSettings->DisplayInternalClasses)
 	{
 		for (const FDirectoryPath& DirPath : InternalPaths)
 		{
 			if (GeneratedClassPathString.StartsWith(DirPath.Path))
 			{
-				bPassesInternalFilter = false;
-				break;
+				return false;
 			}
 		}
 
-		if (bPassesInternalFilter)
+		for (const UClass* Class : InternalClasses)
 		{
-			for (const UClass* Class : InternalClasses)
+			if (InClass->IsChildOf(Class))
 			{
-				if (InClass->IsChildOf(Class))
-				{
-					bPassesInternalFilter = false;
-					break;
-				}
+				return false;
 			}
 		}
 	}
 
 	// The INI files can contain a list of globally allowed classes - if it does, then only classes whose names match will be shown.
-	bool bPassesAllowedClasses = true;
 	if (ClassViewerSettings->AllowedClasses.Num() > 0)
 	{
 		if (!ClassViewerSettings->AllowedClasses.Contains(GeneratedClassPathString))
 		{
-			bPassesAllowedClasses = false;
+			return false;
 		}
 	}
 
-	bool bPassesPlaceableFilter = true;
 	if (InInitOptions.bIsPlaceableOnly)
 	{
-		bPassesPlaceableFilter = IsPlaceable(InClass) && 
-			(InInitOptions.Mode == EClassViewerMode::ClassPicker || !IsBrush(InClass));
+		if (!IsPlaceable(InClass) || 
+			!(InInitOptions.Mode == EClassViewerMode::ClassPicker || !IsBrush(InClass)))
+		{
+			return false;
+		}
 	}
 
-	bool bPassesCustomFilter = true;
+	// REINST classes cannot be used in any class viewer. 
+	if (IsReinstClass(InClass))
+	{
+		return false;
+	}
+	
 	for (const TSharedRef<IClassViewerFilter>& CustomFilter : InInitOptions.ClassFilters)
 	{
 		if (!CustomFilter->IsClassAllowed(InInitOptions, InClass, FilterFunctions))
 		{
-			bPassesCustomFilter = false;
-			break;
+			return false;
 		}
 	}
 
 PRAGMA_DISABLE_DEPRECATION_WARNINGS
 	// Retained for backcompat; can remove when fully deprecated.
-	if (bPassesCustomFilter && InInitOptions.ClassFilter.IsValid())
+	if (InInitOptions.ClassFilter.IsValid())
 	{
-		bPassesCustomFilter = InInitOptions.ClassFilter->IsClassAllowed(InInitOptions, InClass, FilterFunctions);
+		if (!InInitOptions.ClassFilter->IsClassAllowed(InInitOptions, InClass, FilterFunctions))
+		{
+			return false;
+		}
 	}
 PRAGMA_ENABLE_DEPRECATION_WARNINGS
 
-	bool bPassesGlobalClassFilter = true;
 	if (const TSharedPtr<IClassViewerFilter>& GlobalClassFilter = FModuleManager::LoadModuleChecked<FClassViewerModule>("ClassViewer").GetGlobalClassViewerFilter())
 	{
-		bPassesGlobalClassFilter = GlobalClassFilter->IsClassAllowed(InInitOptions, InClass, FilterFunctions);
+		if (!GlobalClassFilter->IsClassAllowed(InInitOptions, InClass, FilterFunctions))
+		{
+			return false;
+		}
 	}
 
-	bool bPassesTextFilter = true;
 	if (bCheckTextFilter && (TextFilter->GetFilterType() != ETextFilterExpressionType::Empty))
 	{
 		FString ClassNameWithCppPrefix = FString::Printf(TEXT("%s%s"), InClass->GetPrefixCPP(), *InClass->GetName());
-		bPassesTextFilter = PassesTextFilter(InClass->GetName(), TextFilter) || PassesTextFilter(ClassNameWithCppPrefix, TextFilter) || PassesTextFilter(InClass->GetDisplayNameText().ToString(), TextFilter);
+		bool bPassesTextFilter = PassesTextFilter(InClass->GetName(), TextFilter) || PassesTextFilter(ClassNameWithCppPrefix, TextFilter) || PassesTextFilter(InClass->GetDisplayNameText().ToString(), TextFilter);
 
 		// If the class is deprecated, try searching without the deprecated name inserted, in case a user typed a string
 		if (!bPassesTextFilter && InClass->HasAnyClassFlags(CLASS_Deprecated))
 		{
 			ClassNameWithCppPrefix.RemoveAt(1, 11);
-			bPassesTextFilter |= PassesTextFilter(ClassNameWithCppPrefix, TextFilter);
+			bPassesTextFilter = PassesTextFilter(ClassNameWithCppPrefix, TextFilter);
+		}
+
+		if (!bPassesTextFilter)
+		{
+			return false;
 		}
 	}
 
-	bool bPassesAssetReferenceFilter = true;
 	if (AssetReferenceFilter.IsValid() && !InClass->IsNative())
 	{
-		bPassesAssetReferenceFilter = AssetReferenceFilter->PassesFilter(FAssetData(InClass));
+		// This check is very slow as it scans all the asset tags
+		if (!AssetReferenceFilter->PassesFilter(FAssetData(InClass)))
+		{
+			return false;
+		}
 	}
 
-	// REINST classes cannot be used in any class viewer. 
-	const bool bPassesReinstFilter = !IsReinstClass(InClass);
-
-	bool bPassesFilter = bPassesAllowedClasses && bPassesPlaceableFilter && bPassesBlueprintBaseFilter
-		&& bPassesDeveloperFilter && bPassesInternalFilter && bPassesEditorClassFilter 
-		&& bPassesCustomFilter && bPassesGlobalClassFilter && bPassesTextFilter && bPassesAssetReferenceFilter
-		&& bPassesReinstFilter;
-
-	return bPassesFilter;
+	return true;
 }
 
 bool FClassViewerFilter::IsUnloadedClassAllowed(const FClassViewerInitializationOptions& InInitOptions, const TSharedRef<const IUnloadedBlueprintData> InUnloadedClassData, TSharedRef<FClassViewerFilterFuncs> InFilterFuncs)
@@ -584,102 +596,115 @@ bool FClassViewerFilter::IsUnloadedClassAllowed(const FClassViewerInitialization
 	}
 
 	const bool bIsBlueprintBase = CheckIfBlueprintBase(InUnloadedClassData);
-	const bool bPassesBlueprintBaseFilter = !InInitOptions.bIsBlueprintBaseOnly || bIsBlueprintBase;
+	if (InInitOptions.bIsBlueprintBaseOnly && !bIsBlueprintBase)
+	{
+		return false;
+	}
+	// TODO: There is currently no good way to handle bEditorClassesOnly for unloaded blueprints
 
 	// Determine if we allow any developer folder classes, if so determine if this class is in one of the allowed developer folders.
 	static const FString DeveloperPathWithSlash = FPackageName::FilenameToLongPackageName(FPaths::GameDevelopersDir());
 	static const FString UserDeveloperPathWithSlash = FPackageName::FilenameToLongPackageName(FPaths::GameUserDeveloperDir());
-	FString GeneratedClassPathString = InUnloadedClassData->GetClassPathName().ToString();
-
-	bool bPassesDeveloperFilter = true;
+	FTopLevelAssetPath UnloadedAssetPath = InUnloadedClassData->GetClassPathName();
+	FString GeneratedClassPathString = UnloadedAssetPath.ToString();
 
 	const UClassViewerSettings* ClassViewerSettings = GetDefault<UClassViewerSettings>();
 	EClassViewerDeveloperType AllowedDeveloperType = ClassViewerSettings->DeveloperFolderType;
 	if (AllowedDeveloperType == EClassViewerDeveloperType::CVDT_None)
 	{
-		bPassesDeveloperFilter = !GeneratedClassPathString.StartsWith(DeveloperPathWithSlash);
+		if (GeneratedClassPathString.StartsWith(DeveloperPathWithSlash))
+		{
+			return false;
+		}
 	}
 	else if (AllowedDeveloperType == EClassViewerDeveloperType::CVDT_CurrentUser)
 	{
-		if (GeneratedClassPathString.StartsWith(DeveloperPathWithSlash))
+		if (GeneratedClassPathString.StartsWith(DeveloperPathWithSlash) && !GeneratedClassPathString.StartsWith(UserDeveloperPathWithSlash))
 		{
-			bPassesDeveloperFilter = GeneratedClassPathString.StartsWith(UserDeveloperPathWithSlash);
+			return false;
 		}
 	}
 
 	// The INI files declare classes and folders that are considered internal only. Does this class match any of those patterns?
 	// INI path: /Script/ClassViewer.ClassViewerProjectSettings
-	bool bPassesInternalFilter = true;
 	if (!ClassViewerSettings->DisplayInternalClasses)
 	{
 		for (const FDirectoryPath& DirPath : InternalPaths)
 		{
 			if (GeneratedClassPathString.StartsWith(DirPath.Path))
 			{
-				bPassesInternalFilter = false;
-				break;
+				return false;
 			}
 		}
 	}
 
 	// The INI files can contain a list of globally allowed classes - if it does, then only classes whose names match will be shown.
-	bool bPassesAllowedClasses = true;
 	if (ClassViewerSettings->AllowedClasses.Num() > 0)
 	{
 		if (!ClassViewerSettings->AllowedClasses.Contains(GeneratedClassPathString))
 		{
-			bPassesAllowedClasses = false;
+			return false;
 		}
 	}
 
-	bool bPassesPlaceableFilter = true;
 	if (InInitOptions.bIsPlaceableOnly)
 	{
-		bPassesPlaceableFilter = IsPlaceable(InUnloadedClassData) &&
-			(InInitOptions.Mode == EClassViewerMode::ClassPicker || !IsBrush(InUnloadedClassData));
+		if (!IsPlaceable(InUnloadedClassData) ||
+			!(InInitOptions.Mode == EClassViewerMode::ClassPicker || !IsBrush(InUnloadedClassData)))
+		{
+			return false;
+		}
 	}
 
-	bool bPassesCustomFilter = true;
 	for (const TSharedRef<IClassViewerFilter>& CustomFilter : InInitOptions.ClassFilters)
 	{
 		if (!CustomFilter->IsUnloadedClassAllowed(InInitOptions, InUnloadedClassData, FilterFunctions))
 		{
-			bPassesCustomFilter = false;
-			break;
+			return false;
 		}
 	}
 
 PRAGMA_DISABLE_DEPRECATION_WARNINGS
 	// Retained for backcompat; can remove when fully deprecated.
-	if (bPassesCustomFilter && InInitOptions.ClassFilter.IsValid())
+	if (InInitOptions.ClassFilter.IsValid())
 	{
-		bPassesCustomFilter = InInitOptions.ClassFilter->IsUnloadedClassAllowed(InInitOptions, InUnloadedClassData, FilterFunctions);
+		if (!InInitOptions.ClassFilter->IsUnloadedClassAllowed(InInitOptions, InUnloadedClassData, FilterFunctions))
+		{
+			return false;
+		}
 	}
 PRAGMA_ENABLE_DEPRECATION_WARNINGS
 
-
-	bool bPassesGlobalClassFilter = true;
 	if (const TSharedPtr<IClassViewerFilter>& GlobalClassFilter = FModuleManager::LoadModuleChecked<FClassViewerModule>("ClassViewer").GetGlobalClassViewerFilter())
 	{
-		bPassesGlobalClassFilter = GlobalClassFilter->IsUnloadedClassAllowed(InInitOptions, InUnloadedClassData, FilterFunctions);
+		if (!GlobalClassFilter->IsUnloadedClassAllowed(InInitOptions, InUnloadedClassData, FilterFunctions))
+		{
+			return false;
+		}
 	}
 
-	const bool bPassesTextFilter = PassesTextFilter(*InUnloadedClassData->GetClassName().Get(), TextFilter);
+	if (bCheckTextFilter && !PassesTextFilter(*InUnloadedClassData->GetClassName().Get(), TextFilter))
+	{
+		return false;
+	}
 
-	bool bPassesAssetReferenceFilter = true;
 	if (AssetReferenceFilter.IsValid() && bIsBlueprintBase)
 	{
-		FString BlueprintAssetPath = GeneratedClassPathString;
-		FAssetData BlueprintAssetData = AssetRegistry.GetAssetByObjectPath(FSoftObjectPath(BlueprintAssetPath));
+		// This query is slow so should be done last
+		const bool bOnlyIncludeOnDiskAssets = true; // If it was in memory it would have used the other filter function
+		FAssetData BlueprintAssetData = AssetRegistry.GetAssetByObjectPath(FSoftObjectPath(UnloadedAssetPath), bOnlyIncludeOnDiskAssets);
 		if (!BlueprintAssetData.IsValid())
 		{
-			BlueprintAssetPath = GeneratedClassPathString.LeftChop(2); // Chop off _C
-			BlueprintAssetData = AssetRegistry.GetAssetByObjectPath(FSoftObjectPath(BlueprintAssetPath));
+			FString BlueprintAssetPath = GeneratedClassPathString.LeftChop(2); // Chop off _C
+			BlueprintAssetData = AssetRegistry.GetAssetByObjectPath(FSoftObjectPath(BlueprintAssetPath), bOnlyIncludeOnDiskAssets);
 		}
 
 		if (BlueprintAssetData.IsValid())
 		{
-			bPassesAssetReferenceFilter = AssetReferenceFilter->PassesFilter(BlueprintAssetData);
+			if (!AssetReferenceFilter->PassesFilter(BlueprintAssetData))
+			{
+				return false;
+			}
 		}
 		else
 		{
@@ -687,9 +712,5 @@ PRAGMA_ENABLE_DEPRECATION_WARNINGS
 		}
 	}
 
-	bool bPassesFilter = bPassesPlaceableFilter && bPassesBlueprintBaseFilter && bPassesDeveloperFilter 
-		&& bPassesInternalFilter && bPassesCustomFilter && bPassesGlobalClassFilter
-		&& (!bCheckTextFilter || bPassesTextFilter) && bPassesAssetReferenceFilter;
-
-	return bPassesFilter;
+	return true;
 }
diff --git a/Engine/Source/Editor/Kismet/Private/BlueprintNamespaceHelper.cpp b/Engine/Source/Editor/Kismet/Private/BlueprintNamespaceHelper.cpp
--- a/Engine/Source/Editor/Kismet/Private/BlueprintNamespaceHelper.cpp
+++ b/Engine/Source/Editor/Kismet/Private/BlueprintNamespaceHelper.cpp
@@ -74,7 +74,7 @@ public:
 			return true;
 		}
 
-		FSoftObjectPath ClassPath(InBlueprint->GetClassPathName().ToString());
+		FSoftObjectPath ClassPath(InBlueprint->GetClassPathName());
 		return CachedNamespaceHelper->IsImportedObject(ClassPath);
 	}
 
diff --git a/Engine/Source/Editor/Kismet/Private/BlueprintNamespaceUtilities.cpp b/Engine/Source/Editor/Kismet/Private/BlueprintNamespaceUtilities.cpp
--- a/Engine/Source/Editor/Kismet/Private/BlueprintNamespaceUtilities.cpp
+++ b/Engine/Source/Editor/Kismet/Private/BlueprintNamespaceUtilities.cpp
@@ -163,14 +163,15 @@ FString FBlueprintNamespaceUtilities::GetObjectNamespace(const FSoftObjectPath&
 		return GetObjectNamespace(Object);
 	}
 
+	const bool bIncludeOnlyOnDiskAssets = true; // The resolve object failed so we know it isn't in memory
 	const FAssetRegistryModule& AssetRegistryModule = FModuleManager::LoadModuleChecked<FAssetRegistryModule>(TEXT("AssetRegistry"));
-	FAssetData AssetData = AssetRegistryModule.Get().GetAssetByObjectPath(InObjectPath);
+	FAssetData AssetData = AssetRegistryModule.Get().GetAssetByObjectPath(InObjectPath, bIncludeOnlyOnDiskAssets);
 	if (!AssetData.IsValid())
 	{
 		FString ObjectPathAsString = InObjectPath.ToString();
 		if (ObjectPathAsString.RemoveFromEnd(TEXT("_C")))
 		{
-			AssetData = AssetRegistryModule.Get().GetAssetByObjectPath(FSoftObjectPath(ObjectPathAsString));
+			AssetData = AssetRegistryModule.Get().GetAssetByObjectPath(FSoftObjectPath(ObjectPathAsString), bIncludeOnlyOnDiskAssets);
 		}
 	}
 
