De-uglifying Class Viewer

diff --git a/Engine/Source/Editor/ClassViewer/Private/SClassViewer.cpp b/Engine/Source/Editor/ClassViewer/Private/SClassViewer.cpp
--- a/Engine/Source/Editor/ClassViewer/Private/SClassViewer.cpp
+++ b/Engine/Source/Editor/ClassViewer/Private/SClassViewer.cpp
@@ -893,6 +893,7 @@ public:
 
 			+SHorizontalBox::Slot()
 				.AutoWidth()
+				.VAlign(VAlign_Center)
 				.Padding( 0.0f, 2.0f, 6.0f, 2.0f )
 				[
 					SNew( SImage )
@@ -918,7 +919,7 @@ public:
 				.AutoWidth()
 				.HAlign(HAlign_Right)
 				.VAlign(VAlign_Center)
-				.Padding( 0.0f, 0.0f, 6.0f, 0.0f )
+				.Padding( 0.0f, 1.0f, 6.0f, 1.0f )
 				[
 					SNew( SComboButton )
 						.ContentPadding(FMargin(2.0f))
