Always scroll the initially selected item into view for class pickers

diff --git a/Engine/Source/Editor/ClassViewer/Private/SClassViewer.cpp b/Engine/Source/Editor/ClassViewer/Private/SClassViewer.cpp
--- a/Engine/Source/Editor/ClassViewer/Private/SClassViewer.cpp
+++ b/Engine/Source/Editor/ClassViewer/Private/SClassViewer.cpp
@@ -2557,6 +2557,7 @@ void SClassViewer::Populate()
 		if (ClassNode)
 		{
 			ClassTree->SetSelection(ClassNode);
+			ClassTree->RequestScrollIntoView(ClassNode);
 		}
 	}
 	else
@@ -2611,6 +2612,7 @@ void SClassViewer::Populate()
 			if(TSharedPtr<FClassViewerNode>* ClassNode = RootTreeItems.FindByPredicate([ClassPathNameToSelect](const TSharedPtr< FClassViewerNode > InClassNode) { return InClassNode->Class.IsValid() && (InClassNode->Class->GetPathName() == ClassPathNameToSelect); }))
 			{
 				ClassList->SetSelection(*ClassNode);
+				ClassList->RequestScrollIntoView(*ClassNode);
 			}
 			InitOptions.InitiallySelectedClass = nullptr;
 		}
