Ignore placeholder classes when adding class viewer nodes.

diff --git a/Engine/Source/Editor/ClassViewer/Private/SClassViewer.cpp b/Engine/Source/Editor/ClassViewer/Private/SClassViewer.cpp
--- a/Engine/Source/Editor/ClassViewer/Private/SClassViewer.cpp
+++ b/Engine/Source/Editor/ClassViewer/Private/SClassViewer.cpp
@@ -1085,6 +1085,7 @@ void FClassHierarchy::CreateNodesForLoadedClasses(TSharedPtr<FClassViewerNode>&
 		UClass* CurrentClass = *ClassIt;
 		// Ignore deprecated and temporary trash classes.
 		if (CurrentClass->HasAnyClassFlags(CLASS_Deprecated | CLASS_NewerVersionExists | CLASS_Hidden) ||
+			FBlueprintSupport::IsClassPlaceholder(CurrentClass) ||
 			FKismetEditorUtilities::IsClassABlueprintSkeleton(CurrentClass))
 		{
 			continue;
