Added perf tag when populating the SClassViewer hierarchy

#rnx

diff --git a/Engine/Source/Editor/ClassViewer/Private/SClassViewer.cpp b/Engine/Source/Editor/ClassViewer/Private/SClassViewer.cpp
--- a/Engine/Source/Editor/ClassViewer/Private/SClassViewer.cpp
+++ b/Engine/Source/Editor/ClassViewer/Private/SClassViewer.cpp
@@ -1388,6 +1388,8 @@ void FClassHierarchy::CreateOrUpdateUnloadedClassNode(TSharedPtr<FClassViewerNod
 
 void FClassHierarchy::PopulateClassHierarchy()
 {
+	TRACE_CPUPROFILER_EVENT_SCOPE(FClassHierarchy::PopulateClassHierarchy);
+
 	// Fetch all classes from AssetRegistry blueprint data (which covers unloaded classes), and in-memory UClasses.
 	// Create a node for each one with unioned data from the AssetRegistry or UClass for that class.
 	// Set parent/child pointers to create a tree, and store this tree in this->ObjectClassRoot
