Follow redirectors when finding parent classes for the class hierarchy so it classes with redirector parents are included in the hierarchy
#tests Opened the Create Blueprint dialog and made sure classes that have a redirector for their parent are selectable

diff --git a/Engine/Source/Editor/ClassViewer/Private/SClassViewer.cpp b/Engine/Source/Editor/ClassViewer/Private/SClassViewer.cpp
--- a/Engine/Source/Editor/ClassViewer/Private/SClassViewer.cpp
+++ b/Engine/Source/Editor/ClassViewer/Private/SClassViewer.cpp
@@ -1395,10 +1395,11 @@ void FClassHierarchy::PopulateClassHierarchy()
 	// Set parent/child pointers to create a tree, and store this tree in this->ObjectClassRoot
 	TMap<FTopLevelAssetPath, TSharedPtr<FClassViewerNode>> ClassPathToNode;
 
+	IAssetRegistry& AssetRegistry = FModuleManager::LoadModuleChecked<FAssetRegistryModule>(AssetRegistryConstants::ModuleName).Get();
+
 	// Create a node for every Blueprint class listed in the AssetRegistry and set the Blueprint fields
 	// Retrieve all blueprint classes 
 	{
-		IAssetRegistry& AssetRegistry = FModuleManager::LoadModuleChecked<FAssetRegistryModule>(AssetRegistryConstants::ModuleName).Get();
 		FString ClassPathString;
 
 		TArray<FAssetData> Assets;
@@ -1449,6 +1450,21 @@ void FClassHierarchy::PopulateClassHierarchy()
 		{
 			ParentNodePtr = ClassPathToNode.Find(Node->ParentClassPath);
 		}
+
+		//Check if the parent is a redirector and use its destination if so
+		if (!ParentNodePtr)
+		{
+			FAssetData ParentData = AssetRegistry.GetAssetByObjectPath(FSoftObjectPath(Node->ParentClassPath));
+
+			if (UObjectRedirector* Redirector = Cast<UObjectRedirector>(ParentData.GetAsset()))
+			{
+				if (UClass* RedirectorClass = Cast<UClass>(Redirector->DestinationObject))
+				{
+					ParentNodePtr = ClassPathToNode.Find(RedirectorClass->GetClassPathName());
+				}
+			}
+		}
+
 		if (!ParentNodePtr)
 		{
 			UE_LOG(LogEditorClassViewer, Warning, TEXT("Class %s has parent %s, but this parent is not found. The Class will not be shown in ClassViewer."),
