Added GetClassFilter metadata tag for class and soft class properties
- This allows specifying a static UFunction with the signature bool (const UClass*) to filter classes in the class picker.

diff --git a/Engine/Source/Editor/ClassViewer/Private/ClassViewerFilter.cpp b/Engine/Source/Editor/ClassViewer/Private/ClassViewerFilter.cpp
--- a/Engine/Source/Editor/ClassViewer/Private/ClassViewerFilter.cpp
+++ b/Engine/Source/Editor/ClassViewer/Private/ClassViewerFilter.cpp
@@ -410,6 +410,25 @@ FClassViewerFilter::FClassViewerFilter(const FClassViewerInitializationOptions&
 		AssetReferenceFilterContext.AddReferencingAssetsFromPropertyHandle(InInitOptions.PropertyHandle);
 		AssetReferenceFilter = GEditor->MakeAssetReferenceFilter(AssetReferenceFilterContext);
 	}
+	
+	// MetaData class filter
+	FProperty* Property = InInitOptions.PropertyHandle.IsValid() ? InInitOptions.PropertyHandle->GetProperty() : nullptr;
+	if (Property && Property->GetOwnerProperty()->HasMetaData("GetClassFilter"))
+	{
+		const FString& GetClassFilterFunctionName = Property->GetOwnerProperty()->GetMetaData("GetClassFilter");
+		if (!GetClassFilterFunctionName.IsEmpty())
+		{
+			if (const UClass* OuterBaseClass = InInitOptions.PropertyHandle->GetOuterBaseClass())
+			{
+				UObject* OuterBaseCDO = OuterBaseClass->GetDefaultObject<UObject>();
+				const UFunction* GetClassFilterFunction = OuterBaseCDO ? OuterBaseCDO->FindFunction(*GetClassFilterFunctionName) : nullptr;
+				if (ensure(GetClassFilterFunction && GetClassFilterFunction->HasAllFunctionFlags(FUNC_Static)))
+				{
+					FilterDelegate = FOnShouldFilterClass::CreateUFunction(OuterBaseCDO, GetClassFilterFunction->GetFName());
+				}
+			}
+		}
+	}
 }
 
 bool FClassViewerFilter::IsNodeAllowed(const FClassViewerInitializationOptions& InInitOptions, const TSharedRef<FClassViewerNode>& InNode, const bool bCheckTextFilter)
@@ -516,6 +535,14 @@ bool FClassViewerFilter::IsClassAllowed(const FClassViewerInitializationOptions&
 	{
 		return false;
 	}
+
+	if (FilterDelegate.IsBound())
+	{
+		if (FilterDelegate.Execute(InClass))
+		{
+			return false;
+		}
+	}
 	
 	for (const TSharedRef<IClassViewerFilter>& CustomFilter : InInitOptions.ClassFilters)
 	{
diff --git a/Engine/Source/Editor/ClassViewer/Public/ClassViewerFilter.h b/Engine/Source/Editor/ClassViewer/Public/ClassViewerFilter.h
--- a/Engine/Source/Editor/ClassViewer/Public/ClassViewerFilter.h
+++ b/Engine/Source/Editor/ClassViewer/Public/ClassViewerFilter.h
@@ -28,6 +28,9 @@ class UObject;
 /** Delegate used to respond to a filter option change. The argument will indicate whether the option was enabled or disabled. */
 DECLARE_DELEGATE_OneParam(FOnClassViewerFilterOptionChanged, bool);
 
+/** Called to check if a class should be filtered out by external code. Return true to exclude the class from the view. */
+DECLARE_DELEGATE_RetVal_OneParam(bool, FOnShouldFilterClass, const UClass* /*Class*/);
+
 /** Used to define a custom class viewer filter option. */
 class FClassViewerFilterOption
 {
@@ -134,7 +137,7 @@ public:
 
 	TArray<UClass*> InternalClasses;
 	TArray<FDirectoryPath> InternalPaths;
-
+	FOnShouldFilterClass FilterDelegate;
 	TSharedRef<FTextFilterExpressionEvaluator> TextFilter;
 	TSharedRef<FClassViewerFilterFuncs> FilterFunctions;
 	TSharedPtr<IAssetReferenceFilter> AssetReferenceFilter;
diff --git a/Engine/Source/Editor/PropertyEditor/Private/PropertyHandleImpl.cpp b/Engine/Source/Editor/PropertyEditor/Private/PropertyHandleImpl.cpp
--- a/Engine/Source/Editor/PropertyEditor/Private/PropertyHandleImpl.cpp
+++ b/Engine/Source/Editor/PropertyEditor/Private/PropertyHandleImpl.cpp
@@ -1,6 +1,7 @@
 // Copyright Epic Games, Inc. All Rights Reserved.
 
 #include "PropertyHandleImpl.h"
+#include "ClassViewerFilter.h"
 #include "GameFramework/Actor.h"
 #include "Editor.h"
 #include "Presentation/PropertyEditor/PropertyEditor.h"
@@ -4539,6 +4540,30 @@ FPropertyAccess::Result FPropertyHandleObject::SetObjectValueFromSelection()
 			const UClass* const SelectedClass = GEditor->GetFirstSelectedClass(RequiredClass);
 			if (SelectedClass)
 			{
+				// MetaData class filter
+				if (NodeProperty && NodeProperty->GetOwnerProperty()->HasMetaData("GetClassFilter"))
+				{
+					const FString& GetClassFilterFunctionName = NodeProperty->GetOwnerProperty()->GetMetaData("GetClassFilter");
+					if (!GetClassFilterFunctionName.IsEmpty())
+					{
+						TArray<UObject*> ObjectList;
+						GetOuterObjects(ObjectList);
+						if (const UClass* OuterBaseClass = GetOuterBaseClass())
+						{
+							UObject* OuterBaseCDO = OuterBaseClass->GetDefaultObject<UObject>();
+							const UFunction* GetClassFilterFunction = OuterBaseCDO ? OuterBaseCDO->FindFunction(*GetClassFilterFunctionName) : nullptr;
+							if (ensure(GetClassFilterFunction && GetClassFilterFunction->HasAllFunctionFlags(FUNC_Static)))
+							{
+								if (FOnShouldFilterClass::CreateUFunction(OuterBaseCDO, GetClassFilterFunction->GetFName()).Execute(SelectedClass))
+								{
+									// The class fails the metadata filters
+									return FPropertyAccess::Fail;
+								}
+							}
+						}
+					}
+				}
+				
 				if (!InterfaceThatMustBeImplemented || SelectedClass->ImplementsInterface(InterfaceThatMustBeImplemented))
 				{
 					FString const ClassPathName = SelectedClass->GetPathName();
diff --git a/Engine/Source/Runtime/CoreUObject/Public/UObject/ObjectMacros.h b/Engine/Source/Runtime/CoreUObject/Public/UObject/ObjectMacros.h
--- a/Engine/Source/Runtime/CoreUObject/Public/UObject/ObjectMacros.h
+++ b/Engine/Source/Runtime/CoreUObject/Public/UObject/ObjectMacros.h
@@ -1520,6 +1520,12 @@ namespace UM
 		/// bool FuncName(const FAssetData& AssetData) const;
 		GetAssetFilter,
 		
+		/// [PropertyMetadata] Causes classes to be filtered through a static UFunction, e.g. meta=(GetClassFilter="FuncName"). The UFunction should return true to exclude the class.
+		///
+		/// UFUNCTION()
+		/// bool FuncName(const UClass* Class) const;
+		GetClassFilter,
+		
 		/// [PropertyMetadata] Get allowed classes through a UFunction, e.g. meta=(GetAllowedClasses="FuncName").
 		///
 		/// UFUNCTION()
