Use UAssetRegistryHelpers::FixupRedirectedAssetPath instead of GetAsset to resolve asset based redirectors for parent classes. The latter can load a uasset, which is undesirable.

#rnx

diff --git a/Engine/Source/Editor/ClassViewer/Private/SClassViewer.cpp b/Engine/Source/Editor/ClassViewer/Private/SClassViewer.cpp
--- a/Engine/Source/Editor/ClassViewer/Private/SClassViewer.cpp
+++ b/Engine/Source/Editor/ClassViewer/Private/SClassViewer.cpp
@@ -5,6 +5,7 @@
 #include "AddToProjectConfig.h"
 #include "AssetDiscoveryIndicator.h"
 #include "AssetRegistry/AssetData.h"
+#include "AssetRegistry/AssetRegistryHelpers.h"
 #include "AssetRegistry/AssetRegistryModule.h"
 #include "AssetRegistry/IAssetRegistry.h"
 #include "AssetToolsModule.h"
@@ -1395,11 +1396,10 @@ void FClassHierarchy::PopulateClassHierarchy()
 	// Set parent/child pointers to create a tree, and store this tree in this->ObjectClassRoot
 	TMap<FTopLevelAssetPath, TSharedPtr<FClassViewerNode>> ClassPathToNode;
 
-	IAssetRegistry& AssetRegistry = FModuleManager::LoadModuleChecked<FAssetRegistryModule>(AssetRegistryConstants::ModuleName).Get();
-
 	// Create a node for every Blueprint class listed in the AssetRegistry and set the Blueprint fields
 	// Retrieve all blueprint classes 
 	{
+		IAssetRegistry& AssetRegistry = FModuleManager::LoadModuleChecked<FAssetRegistryModule>(AssetRegistryConstants::ModuleName).Get();
 		FString ClassPathString;
 
 		TArray<FAssetData> Assets;
@@ -1449,22 +1449,13 @@ void FClassHierarchy::PopulateClassHierarchy()
 		if (!Node->ParentClassPath.IsNull())
 		{
 			ParentNodePtr = ClassPathToNode.Find(Node->ParentClassPath);
-		}
-
-		//Check if the parent is a redirector and use its destination if so
-		if (!ParentNodePtr)
-		{
-			FAssetData ParentData = AssetRegistry.GetAssetByObjectPath(FSoftObjectPath(Node->ParentClassPath));
-
-			if (UObjectRedirector* Redirector = Cast<UObjectRedirector>(ParentData.GetAsset()))
+			if(!ParentNodePtr)
 			{
-				if (UClass* RedirectorClass = Cast<UClass>(Redirector->DestinationObject))
-				{
-					ParentNodePtr = ClassPathToNode.Find(RedirectorClass->GetClassPathName());
-				}
+				FSoftObjectPath SoftParentClass(Node->ParentClassPath);
+				UAssetRegistryHelpers::FixupRedirectedAssetPath(SoftParentClass);
+				ParentNodePtr = ClassPathToNode.Find(FTopLevelAssetPath(SoftParentClass.ToString()));
 			}
 		}
-
 		if (!ParentNodePtr)
 		{
 			UE_LOG(LogEditorClassViewer, Warning, TEXT("Class %s has parent %s, but this parent is not found. The Class will not be shown in ClassViewer."),
